[PCG] New PCGGeometryScriptInterop plugin, to use geometry script with PCG, in another plugin to avoid forcing the dependency between PCG and GeometryScripts.

* Create a new plugin for interop with GeometryScripts
* Create a new native node MeshSampler
* Same functionality as MeshToPoints node (defined in Blueprint)
* Also add Poisson Sampling (coming from GeometryScript)
* Node is not framebound for execution, useful when sampling is expensive.
* Add condition to only check sleeping tasks once per frame, to avoid spinning during our whole PCG frametime budget.
* Needed to export a few things after PCGAsync revamp.

#rnx

diff --git a/Engine/Plugins/Experimental/PCG/Source/PCG/Private/Graph/PCGGraphExecutor.cpp b/Engine/Plugins/Experimental/PCG/Source/PCG/Private/Graph/PCGGraphExecutor.cpp
--- a/Engine/Plugins/Experimental/PCG/Source/PCG/Private/Graph/PCGGraphExecutor.cpp
+++ b/Engine/Plugins/Experimental/PCG/Source/PCG/Private/Graph/PCGGraphExecutor.cpp
@@ -442,6 +442,7 @@ void FPCGGraphExecutor::Execute()
 
 	// TODO: change this if we support tasks that are not framebound
 	bool bAnyTaskEnded = false;
+	bool bHasAlreadyCheckedSleepingTasks = false;
 
 	const double StartTime = FPlatformTime::Seconds();
 
@@ -464,8 +465,12 @@ void FPCGGraphExecutor::Execute()
 	UpdateGenerationNotification();
 #endif
 
-	while(ReadyTasks.Num() > 0 || ActiveTasks.Num() > 0 || SleepingTasks.Num() > 0)
+	while(ReadyTasks.Num() > 0 || ActiveTasks.Num() > 0 || (!bHasAlreadyCheckedSleepingTasks && SleepingTasks.Num() > 0))
 	{
+		// If we only have sleeping tasks, we will go through this loop only once. If all tasks are still sleeping after one iteration,
+		// we will wake them up only at next tick. It will avoid spinning for our whole frametime budget.
+		bHasAlreadyCheckedSleepingTasks = ReadyTasks.Num() == 0 && ActiveTasks.Num() == 0 && SleepingTasks.Num() > 0;
+
 		// First: if we have free resources, move ready tasks to the active tasks
 		bool bMainThreadAvailable = (ActiveTasks.Num() == 0 || !ActiveTasks[0].Element->CanExecuteOnlyOnMainThread(ActiveTasks[0].Context.Get()));
 		int32 NumAvailableThreads = FMath::Max(0, MaxNumThreads - CurrentlyUsedThreads);
diff --git a/Engine/Plugins/Experimental/PCG/Source/PCG/Public/Helpers/PCGAsync.h b/Engine/Plugins/Experimental/PCG/Source/PCG/Public/Helpers/PCGAsync.h
--- a/Engine/Plugins/Experimental/PCG/Source/PCG/Public/Helpers/PCGAsync.h
+++ b/Engine/Plugins/Experimental/PCG/Source/PCG/Public/Helpers/PCGAsync.h
@@ -17,10 +17,10 @@ struct FPCGPoint;
 
 namespace FPCGAsync
 {
-	struct ConsoleVar
+	namespace ConsoleVar
 	{
-		static TAutoConsoleVariable<bool> CVarDisableAsyncTimeSlicing;
-		static TAutoConsoleVariable<int32> CVarAsyncOverrideChunkSize;
+		extern PCG_API TAutoConsoleVariable<bool> CVarDisableAsyncTimeSlicing;
+		extern PCG_API TAutoConsoleVariable<int32> CVarAsyncOverrideChunkSize;
 	};
 
 	/** 
diff --git a/Engine/Plugins/Experimental/PCG/Source/PCG/Public/Helpers/PCGAsyncState.h b/Engine/Plugins/Experimental/PCG/Source/PCG/Public/Helpers/PCGAsyncState.h
--- a/Engine/Plugins/Experimental/PCG/Source/PCG/Public/Helpers/PCGAsyncState.h
+++ b/Engine/Plugins/Experimental/PCG/Source/PCG/Public/Helpers/PCGAsyncState.h
@@ -26,5 +26,5 @@ struct FPCGAsyncState
 	bool bIsRunningAsyncCall = false;
 
 	/** Returns true if we reached end time. */
-	bool ShouldStop() const;
+	PCG_API bool ShouldStop() const;
 };
\ No newline at end of file
diff --git a/Engine/Plugins/Experimental/PCGGeometryScriptInterop/PCGGeometryScriptInterop.uplugin b/Engine/Plugins/Experimental/PCGGeometryScriptInterop/PCGGeometryScriptInterop.uplugin
new file mode 100644
--- /dev/null
+++ b/Engine/Plugins/Experimental/PCGGeometryScriptInterop/PCGGeometryScriptInterop.uplugin
@@ -0,0 +1,34 @@
+{
+	"FileVersion": 3,
+	"Version": 1,
+	"VersionName": "0.1",
+	"FriendlyName": "Procedural Content Generation Framework Geometry Script Interop",
+	"Description": "Extra plugin for Procedural Content Generation Framework interacting with Geometry Scripts.",
+	"Category": "Editor",
+	"CreatedBy": "Epic Games, Inc.",
+	"CreatedByURL": "https://epicgames.com",
+	"DocsURL": "",
+	"MarketplaceURL": "",
+	"SupportURL": "",
+	"EnabledByDefault": false,
+	"CanContainContent": true,
+	"IsBetaVersion": false,
+	"Installed": false,
+	"Modules": [
+		{
+			"Name": "PCGGeometryScriptInterop",
+			"Type": "Runtime",
+			"LoadingPhase": "Default"
+		}
+	],
+	"Plugins": [
+		{
+			"Name": "PCG",
+			"Enabled": true
+		},
+		{
+			"Name": "GeometryScripting",
+			"Enabled": true
+		}
+	]
+}
diff --git a/Engine/Plugins/Experimental/PCGGeometryScriptInterop/Source/PCGGeometryScriptInterop/PCGGeometryScriptInterop.Build.cs b/Engine/Plugins/Experimental/PCGGeometryScriptInterop/Source/PCGGeometryScriptInterop/PCGGeometryScriptInterop.Build.cs
new file mode 100644
--- /dev/null
+++ b/Engine/Plugins/Experimental/PCGGeometryScriptInterop/Source/PCGGeometryScriptInterop/PCGGeometryScriptInterop.Build.cs
@@ -0,0 +1,42 @@
+// Copyright Epic Games, Inc. All Rights Reserved.
+
+namespace UnrealBuildTool.Rules
+{
+	public class PCGGeometryScriptInterop : ModuleRules
+	{
+		public PCGGeometryScriptInterop(ReadOnlyTargetRules Target) : base(Target)
+		{
+			PCHUsage = PCHUsageMode.UseExplicitOrSharedPCHs;
+
+			PublicDependencyModuleNames.AddRange(
+				new string[] {
+					"Core",
+					"CoreUObject",
+					"Engine",
+					"GeometryScriptingCore",
+					"Projects",
+					"RenderCore",
+					"RHI",
+					"PCG"
+				});
+
+			PrivateDependencyModuleNames.AddRange(
+				new string[]
+				{
+					"GeometryCore",
+					"GeometryFramework"
+				}
+			);
+
+			if (Target.bBuildEditor == true)
+			{
+				PrivateDependencyModuleNames.AddRange(
+					new string[]
+					{
+						"UnrealEd",
+					}
+				);
+			}
+		}
+	}
+}
diff --git a/Engine/Plugins/Experimental/PCGGeometryScriptInterop/Source/PCGGeometryScriptInterop/Private/Elements/PCGMeshSampler.cpp b/Engine/Plugins/Experimental/PCGGeometryScriptInterop/Source/PCGGeometryScriptInterop/Private/Elements/PCGMeshSampler.cpp
new file mode 100644
--- /dev/null
+++ b/Engine/Plugins/Experimental/PCGGeometryScriptInterop/Source/PCGGeometryScriptInterop/Private/Elements/PCGMeshSampler.cpp
@@ -0,0 +1,342 @@
+// Copyright Epic Games, Inc. All Rights Reserved.
+
+#include "Elements/PCGMeshSampler.h"
+
+#include "PCGComponent.h"
+#include "PCGModule.h"
+#include "PCGPin.h"
+#include "Data/PCGPointData.h"
+#include "Helpers/PCGAsync.h"
+#include "Helpers/PCGBlueprintHelpers.h"
+
+#include "UDynamicMesh.h"
+#include "Engine/AssetManager.h"
+#include "Engine/StaticMesh.h"
+#include "GameFramework/Actor.h"
+#include "GeometryScript/MeshAssetFunctions.h"
+#include "GeometryScript/MeshNormalsFunctions.h"
+#include "GeometryScript/MeshQueryFunctions.h"
+#include "GeometryScript/MeshRepairFunctions.h"
+#include "GeometryScript/MeshVertexColorFunctions.h"
+#include "GeometryScript/MeshVoxelFunctions.h"
+#include "Sampling/MeshSurfacePointSampling.h"
+
+#define LOCTEXT_NAMESPACE "PCGMeshSampler"
+
+UPCGMeshSamplerSettings::UPCGMeshSamplerSettings()
+	: UPCGSettings()
+{
+	bUseSeed = true;
+}
+
+TArray<FPCGPinProperties> UPCGMeshSamplerSettings::InputPinProperties() const
+{
+	return TArray<FPCGPinProperties>{};
+}
+
+FPCGElementPtr UPCGMeshSamplerSettings::CreateElement() const
+{
+	return MakeShared<FPCGMeshSamplerElement>();
+}
+
+#if WITH_EDITOR
+FText UPCGMeshSamplerSettings::GetNodeTooltipText() const
+{
+	return LOCTEXT("MeshSamplerNodeTooltip", "Sample points on a static mesh.");
+}
+#endif // WITH_EDITOR
+
+FPCGMeshSamplerContext::~FPCGMeshSamplerContext()
+{
+	// The context can be destroyed if the task is canceled. In that case, we need to notify the future that it should stop,
+	// and wait for it to finish (as the future uses some data stored in the context that will go dangling if the context is destroyed).
+	if (SamplingFuture.IsValid() && !SamplingFuture.IsReady())
+	{
+		StopSampling = true;
+		SamplingFuture.Wait();
+	}
+
+	// DynamicMesh needs to be removed from root, as we added it to root when we created it to avoid getting GC between executions.
+	if (DynamicMesh)
+	{
+		DynamicMesh->RemoveFromRoot();
+		DynamicMesh->MarkAsGarbage();
+	}
+}
+
+FPCGContext* FPCGMeshSamplerElement::Initialize(const FPCGDataCollection& InputData, TWeakObjectPtr<UPCGComponent> SourceComponent, const UPCGNode* Node)
+{
+	FPCGMeshSamplerContext* Context = new FPCGMeshSamplerContext();
+	Context->InputData = InputData;
+	Context->SourceComponent = SourceComponent;
+	Context->Node = Node;
+
+	return Context;
+}
+
+bool FPCGMeshSamplerElement::PrepareDataInternal(FPCGContext* InContext) const
+{
+	TRACE_CPUPROFILER_EVENT_SCOPE(FPCGMeshSamplerElement::PrepareData);
+	FPCGMeshSamplerContext* Context = static_cast<FPCGMeshSamplerContext*>(InContext);
+
+	check(Context);
+
+	if (!Context->SourceComponent.IsValid())
+	{
+		return true;
+	}
+
+	const UPCGMeshSamplerSettings* Settings = Context->GetInputSettings<UPCGMeshSamplerSettings>();
+	check(Settings);
+
+	// TODO: Could be async
+	TSoftObjectPtr<UStaticMesh> StaticMeshPtr{ Settings->StaticMeshPath };
+	UStaticMesh* StaticMesh = StaticMeshPtr.LoadSynchronous();
+
+	if (!StaticMesh)
+	{
+		PCGE_LOG(Error, "Provided static mesh does not exist: %s", *Settings->StaticMeshPath.ToString());
+		return true;
+	}
+
+	Context->DynamicMesh = NewObject<UDynamicMesh>();
+	// Need to add to root to make sure we don't delete it. Will be removed from root in the context destructor.
+	Context->DynamicMesh->AddToRoot();
+
+	EGeometryScriptOutcomePins Outcome;
+
+	UGeometryScriptLibrary_StaticMeshFunctions::CopyMeshFromStaticMesh(StaticMesh, Context->DynamicMesh, FGeometryScriptCopyMeshFromAssetOptions{},
+		FGeometryScriptMeshReadLOD{ Settings->RequestedLODType, Settings->RequestedLODIndex }, Outcome);
+
+	if (Outcome == EGeometryScriptOutcomePins::Failure)
+	{
+		PCGE_LOG(Error, "Static to Dynamic mesh failed");
+		return true;
+	}
+	
+	if (Settings->bVoxelize)
+	{
+		FGeometryScriptSolidifyOptions SolidifyOptions{};
+		SolidifyOptions.GridParameters.GridCellSize = Settings->VoxelSize;
+		SolidifyOptions.GridParameters.SizeMethod = EGeometryScriptGridSizingMethod::GridCellSize;
+
+		UGeometryScriptLibrary_MeshVoxelFunctions::ApplyMeshSolidify(Context->DynamicMesh, SolidifyOptions);
+
+		if (Settings->bRemoveHiddenTriangles)
+		{
+			FGeometryScriptRemoveHiddenTrianglesOptions RemoveTriangleOptions{};
+			
+			UGeometryScriptLibrary_MeshRepairFunctions::RemoveHiddenTriangles(Context->DynamicMesh, RemoveTriangleOptions);
+		}
+	}
+
+	switch (Settings->SamplingMethod)
+	{
+	case EPCGMeshSamplingMethod::OnePointPerVertex:
+	{
+		bool Dummy, Dummy2;
+		UGeometryScriptLibrary_MeshQueryFunctions::GetAllVertexPositions(Context->DynamicMesh, Context->Positions, /*bSkipGaps=*/false, /*bHasVertexIDGaps=*/Dummy);
+		UGeometryScriptLibrary_MeshVertexColorFunctions::GetMeshPerVertexColors(Context->DynamicMesh, Context->Colors, /*bIsValidColorSet=*/Dummy, /*bHasVertexIDGaps=*/Dummy2);
+		UGeometryScriptLibrary_MeshNormalsFunctions::GetMeshPerVertexNormals(Context->DynamicMesh, Context->Normals, /*bIsValidNormalSet=*/Dummy, /*bHasVertexIDGaps=*/Dummy2);
+		Context->Iterations = Context->Positions.List->Num();
+		break;
+	}
+	case EPCGMeshSamplingMethod::OnePointPerTriangle:
+	{
+		bool Dummy;
+		UGeometryScriptLibrary_MeshQueryFunctions::GetAllTriangleIDs(Context->DynamicMesh, Context->TriangleIds, /*bHasTriangleIDGaps=*/Dummy);
+		Context->Iterations = Context->TriangleIds.List->Num();
+		break;
+	}
+	case EPCGMeshSamplingMethod::PoissonSampling:
+	{
+		// No preparation needed
+		break;
+	}
+	default:
+	{
+		return true;
+	}
+	}
+
+	TArray<FPCGTaggedData>& Outputs = InContext->OutputData.TaggedData;
+	Context->OutPointData = NewObject<UPCGPointData>();
+	Context->OutPointData->TargetActor = Context->SourceComponent->GetOwner();
+	Outputs.Emplace_GetRef().Data = Context->OutPointData;
+
+	Context->bDataPrepared = true;
+	return true;
+}
+
+bool FPCGMeshSamplerElement::ExecuteInternal(FPCGContext* InContext) const
+{
+	TRACE_CPUPROFILER_EVENT_SCOPE(FPCGMeshSamplerElement::Execute);
+	FPCGMeshSamplerContext* Context = static_cast<FPCGMeshSamplerContext*>(InContext);
+
+	check(Context);
+
+	if (!Context->bDataPrepared)
+	{
+		return true;
+	}
+
+	check(Context->OutPointData);
+
+	const UPCGMeshSamplerSettings* Settings = Context->GetInputSettings<UPCGMeshSamplerSettings>();
+	check(Settings);
+
+	bool bIsDone = false;
+	const bool bEnableTimeSlicing = true;
+	const int Seed = Context->GetSeed();
+
+	switch (Settings->SamplingMethod)
+	{
+	case EPCGMeshSamplingMethod::OnePointPerVertex:
+	{
+		TRACE_CPUPROFILER_EVENT_SCOPE(FPCGMeshSamplerElement::Execute::OnePointPerVertex);
+
+		const TArray<FVector>& Positions = *Context->Positions.List.Get();
+		const TArray<FLinearColor>& Colors = *Context->Colors.List.Get();
+		const TArray<FVector>& Normals = *Context->Normals.List.Get();
+
+		auto IterationBody = [&Positions, &Colors, &Normals, Settings](int32 Index, FPCGPoint& OutPoint) -> bool
+		{
+			const FVector& Position = Positions[Index];
+			const FLinearColor& Color = Colors[Index];
+			const FVector& Normal = Normals[Index];
+
+			OutPoint = FPCGPoint{};
+			OutPoint.Transform = FTransform{ FRotationMatrix::MakeFromZ(Normal).Rotator(), Position, FVector{1.0, 1.0, 1.0} };
+			OutPoint.Density = Settings->bUseRedAsDensity ? Color.R : 1.0f;
+			OutPoint.Color = Color;
+
+			UPCGBlueprintHelpers::SetSeedFromPosition(OutPoint);
+
+			return true;
+		};
+
+		bIsDone = FPCGAsync::AsyncProcessing<FPCGPoint>(&Context->AsyncState, Context->Iterations, Context->OutPointData->GetMutablePoints(), IterationBody, bEnableTimeSlicing);
+		break;
+	}
+	case EPCGMeshSamplingMethod::OnePointPerTriangle:
+	{
+		TRACE_CPUPROFILER_EVENT_SCOPE(FPCGMeshSamplerElement::Execute::OnePointPerTriangle);
+
+		const TArray<int32>& TriangleIds = *Context->TriangleIds.List.Get();
+
+		auto IterationBody = [DynamicMesh = Context->DynamicMesh, &TriangleIds](int32 Index, FPCGPoint& OutPoint) -> bool
+		{
+			const int32 TriangleId = TriangleIds[Index];
+
+			FVector Vertex1{}, Vertex2{}, Vertex3{};
+			bool bIsValidTriangle;
+
+			UGeometryScriptLibrary_MeshQueryFunctions::GetTrianglePositions(DynamicMesh, TriangleId, bIsValidTriangle, Vertex1, Vertex2, Vertex3);
+			const FVector Normal = UGeometryScriptLibrary_MeshQueryFunctions::GetTriangleFaceNormal(DynamicMesh, TriangleId, bIsValidTriangle);
+
+			const FVector Position = (Vertex1 + Vertex2 + Vertex3) / 3.0;
+
+			OutPoint = FPCGPoint{};
+			OutPoint.Transform = FTransform{ FRotationMatrix::MakeFromZ(Normal).Rotator(), Position, FVector{1.0, 1.0, 1.0} };
+			OutPoint.Density = 1.0f;
+
+			UPCGBlueprintHelpers::SetSeedFromPosition(OutPoint);
+
+			return true;
+		};
+
+		bIsDone = FPCGAsync::AsyncProcessing<FPCGPoint>(&Context->AsyncState, Context->Iterations, Context->OutPointData->GetMutablePoints(), IterationBody, bEnableTimeSlicing);
+		break;
+	}
+	case EPCGMeshSamplingMethod::PoissonSampling:
+	{
+		// For Poisson sampling, we are calling a "all-in-one" function, where we don't have control for timeslicing.
+		// Since Poisson sampling can be expensive (depending on the radius used), we will do the sampling in a future, put this task to sleep,
+		// and wait for the sampling to wake us up.
+		if (!Context->SamplingFuture.IsValid())
+		{
+			// Put the task asleep
+			Context->bIsPaused = true;
+			Context->SamplingProgess = MakeUnique<FProgressCancel>();
+			Context->SamplingProgess->CancelF = [Context]() -> bool { return Context->StopSampling; };
+
+			auto SamplingFuture = [Settings, Context, Seed]() -> bool
+			{
+				TRACE_CPUPROFILER_EVENT_SCOPE(FPCGMeshSamplerElement::Execute::PoissonSampling);
+
+				// Taken from PCGHelpers, which is private.
+				auto ComputeSeed = [](int A, int B) -> int
+				{
+					return ((A * 196314165U) + 907633515U) ^ ((B * 73148459U) + 453816763U);
+				};
+
+				UE::Geometry::FMeshSurfacePointSampling PointSampling;
+				PointSampling.SampleRadius = Settings->SamplingOptions.SamplingRadius;
+				PointSampling.MaxSamples = Settings->SamplingOptions.MaxNumSamples;
+				PointSampling.RandomSeed = ComputeSeed(Seed, Settings->SamplingOptions.RandomSeed);
+				PointSampling.SubSampleDensity = Settings->SamplingOptions.SubSampleDensity;
+
+				if (Settings->NonUniformSamplingOptions.MaxSamplingRadius > PointSampling.SampleRadius)
+				{
+					PointSampling.MaxSampleRadius = Settings->NonUniformSamplingOptions.MaxSamplingRadius;
+					PointSampling.SizeDistribution = UE::Geometry::FMeshSurfacePointSampling::ESizeDistribution(static_cast<int>(Settings->NonUniformSamplingOptions.SizeDistribution));
+					PointSampling.SizeDistributionPower = FMath::Clamp(Settings->NonUniformSamplingOptions.SizeDistributionPower, 1.0, 10.0);
+				}
+
+				PointSampling.ComputePoissonSampling(Context->DynamicMesh->GetMeshRef(), Context->SamplingProgess.Get());
+
+				if (Context->StopSampling)
+				{
+					return true;
+				}
+
+				TArray<FPCGPoint>& Points = Context->OutPointData->GetMutablePoints();
+				Points.Reserve(PointSampling.Samples.Num());
+
+				int Count = 0;
+				for (UE::Geometry::FFrame3d& Sample : PointSampling.Samples)
+				{
+					// Avoid to check too many times
+					constexpr int CancelledCheckNum = 25;
+					if (++Count == CancelledCheckNum)
+					{
+						Count = 0;
+						if (Context->StopSampling)
+						{
+							return true;
+						}
+					}
+
+					FPCGPoint& OutPoint = Points.Emplace_GetRef();
+					OutPoint.Transform = Sample.ToTransform();
+					OutPoint.Density = 1.0f;
+
+					UPCGBlueprintHelpers::SetSeedFromPosition(OutPoint);
+				}
+
+				// Unpause the task
+				Context->bIsPaused = false;
+				return true;
+			};
+
+			Context->SamplingFuture = Async(EAsyncExecution::ThreadPool, std::move(SamplingFuture));
+		}
+
+		bIsDone = Context->SamplingFuture.IsReady();
+		if (bIsDone)
+		{
+			Context->SamplingFuture.Reset();
+		}
+
+		break;
+	}
+	default:
+	{
+		return true;
+	}
+	}
+
+	return bIsDone;
+}
+
+#undef LOCTEXT_NAMESPACE
\ No newline at end of file
diff --git a/Engine/Plugins/Experimental/PCGGeometryScriptInterop/Source/PCGGeometryScriptInterop/Private/PCGGeometryScriptInteropModule.cpp b/Engine/Plugins/Experimental/PCGGeometryScriptInterop/Source/PCGGeometryScriptInterop/Private/PCGGeometryScriptInteropModule.cpp
new file mode 100644
--- /dev/null
+++ b/Engine/Plugins/Experimental/PCGGeometryScriptInterop/Source/PCGGeometryScriptInterop/Private/PCGGeometryScriptInteropModule.cpp
@@ -0,0 +1,16 @@
+// Copyright Epic Games, Inc. All Rights Reserved.
+
+#include "Modules/ModuleManager.h"
+
+class FPCGGeometryScriptInteropModule final : public IModuleInterface
+{
+public:
+	//~ IModuleInterface implementation
+	virtual bool SupportsDynamicReloading() override
+	{
+		return true;
+	}
+	//~ End IModuleInterface implementation
+};
+
+IMPLEMENT_MODULE(FPCGGeometryScriptInteropModule, PCGGeometryScriptInterop);
diff --git a/Engine/Plugins/Experimental/PCGGeometryScriptInterop/Source/PCGGeometryScriptInterop/Public/Elements/PCGMeshSampler.h b/Engine/Plugins/Experimental/PCGGeometryScriptInterop/Source/PCGGeometryScriptInterop/Public/Elements/PCGMeshSampler.h
new file mode 100644
--- /dev/null
+++ b/Engine/Plugins/Experimental/PCGGeometryScriptInterop/Source/PCGGeometryScriptInterop/Public/Elements/PCGMeshSampler.h
@@ -0,0 +1,141 @@
+// Copyright Epic Games, Inc. All Rights Reserved.
+
+#pragma once
+
+#include "PCGContext.h"
+#include "PCGElement.h"
+#include "PCGSettings.h"
+
+#include "Async/Future.h"
+#include "GeometryScript/GeometryScriptTypes.h"
+#include "GeometryScript/MeshSamplingFunctions.h"
+#include "Templates/UniquePtr.h"
+#include "UObject/ObjectPtr.h"
+#include "UObject/SoftObjectPtr.h"
+
+#include <atomic>
+
+#include "PCGMeshSampler.generated.h"
+
+class FProgressCancel;
+class UDynamicMesh;
+class UPCGPointData;
+class UStaticMesh;
+
+UENUM(BlueprintType)
+enum class EPCGMeshSamplingMethod
+{
+	/** Sample one point (at the center) of each triangle of the mesh. */
+	OnePointPerTriangle,
+
+	/** Sample one point per vertex on the mesh. */
+	OnePointPerVertex,
+
+	/** Use Poisson sampling to sample points on the mesh. Can be expensive and therefore it is not framebound. */
+	PoissonSampling
+};
+
+/**
+* Sample points on a mesh
+*/
+UCLASS(BlueprintType, ClassGroup = (Procedural))
+class PCGGEOMETRYSCRIPTINTEROP_API UPCGMeshSamplerSettings : public UPCGSettings
+{
+	GENERATED_BODY()
+
+public:
+	UPCGMeshSamplerSettings();
+
+	//~Begin UPCGSettings interface
+#if WITH_EDITOR
+	virtual FName GetDefaultNodeName() const override { return FName(TEXT("Mesh Sampler")); }
+	virtual FText GetNodeTooltipText() const override;
+	virtual EPCGSettingsType GetType() const override { return EPCGSettingsType::Sampler; }
+#endif
+
+	virtual TArray<FPCGPinProperties> InputPinProperties() const override;
+	virtual TArray<FPCGPinProperties> OutputPinProperties() const override { return DefaultPointOutputPinProperties(); }
+
+protected:
+	virtual FPCGElementPtr CreateElement() const override;
+	// ~End UPCGSettings interface
+
+public:
+	UPROPERTY(BlueprintReadWrite, EditAnywhere, Category = "Settings")
+	EPCGMeshSamplingMethod SamplingMethod = EPCGMeshSamplingMethod::OnePointPerTriangle;
+
+	/** Soft Object Path to the mesh to sample from. Will be loaded. */
+	UPROPERTY(BlueprintReadWrite, EditAnywhere, Category = "Settings", meta = (PCG_Overridable, AllowedClasses = "/Script/Engine.StaticMesh", DisplayName = "Static Mesh"))
+	FSoftObjectPath StaticMeshPath;
+
+	/** In "One Point Per Vertex" option, will assign point density from the red component of the vertex color. */
+	UPROPERTY(BlueprintReadWrite, EditAnywhere, Category = "Settings|Per-Vertex Options", meta = (PCG_Overridable, EditCondition = "SamplingMethod == EPCGMeshSamplingMethod::OnePointPerVertex"))
+	bool bUseRedAsDensity = false;
+
+	/** Enable voxelisation as a preparation pass. Can be more expensive given the VoxelSize. */
+	UPROPERTY(BlueprintReadWrite, EditAnywhere, Category = "Settings|Voxelize Options", meta = (PCG_Overridable))
+	bool bVoxelize = false;
+
+	/** Size of a voxel for the voxelization. */
+	UPROPERTY(BlueprintReadWrite, EditAnywhere, Category = "Settings|Voxelize Options", meta = (PCG_Overridable, EditCondition = "bVoxelize"))
+	float VoxelSize = 100.0f;
+
+	/** Post-processing pass after voxelization to remove hidden triangles. */
+	UPROPERTY(BlueprintReadWrite, EditAnywhere, Category = "Settings|Voxelize Options", meta = (PCG_Overridable, EditCondition = "bVoxelize"))
+	bool bRemoveHiddenTriangles = true;
+
+	/** LOD type to use when creating DynamicMesh from specified StaticMesh. */
+	UPROPERTY(BlueprintReadWrite, EditAnywhere, Category = "Settings|LODSettings", meta = (PCG_Overridable))
+	EGeometryScriptLODType RequestedLODType = EGeometryScriptLODType::RenderData;
+
+	UPROPERTY()
+	int32 RequestedLODIndex = 0;
+
+	// Poisson Sampling parameters
+	UPROPERTY(BlueprintReadWrite, EditAnywhere, Category = "Settings|Poisson sampling", meta = (PCG_Overridable, EditCondition = "SamplingMethod == EPCGMeshSamplingMethod::PoissonSampling"))
+	FGeometryScriptMeshPointSamplingOptions SamplingOptions;
+
+	UPROPERTY(BlueprintReadWrite, EditAnywhere, Category = "Settings|Poisson sampling", meta = (PCG_Overridable, EditCondition = "SamplingMethod == EPCGMeshSamplingMethod::PoissonSampling"))
+	FGeometryScriptNonUniformPointSamplingOptions NonUniformSamplingOptions;
+};
+
+/**
+* Extra context to store all the data that need to be kept between multiple executions (time slicing)
+*/
+struct FPCGMeshSamplerContext : public FPCGContext
+{
+	~FPCGMeshSamplerContext();
+	
+	// Dynamic mesh. Will be added to root.
+	TObjectPtr<UDynamicMesh> DynamicMesh;
+
+	// Lists extracted from the mesh
+	FGeometryScriptVectorList Positions;
+	FGeometryScriptColorList Colors;
+	FGeometryScriptVectorList Normals;
+	FGeometryScriptIndexList TriangleIds;
+
+	// Output point data
+	UPCGPointData* OutPointData = nullptr;
+
+	// For Poisson sampling, we are starting a future that is not framebound
+	// Store the future and synchronisation items in the context
+	TFuture<bool> SamplingFuture;
+	std::atomic<bool> StopSampling = false;
+	TUniquePtr<FProgressCancel> SamplingProgess;
+
+	// Number of iterations to be done
+	int32 Iterations = 0;
+
+	// Set to true if prepared data succeeded.
+	bool bDataPrepared = false;
+};
+
+class FPCGMeshSamplerElement : public IPCGElement
+{
+public:
+	virtual FPCGContext* Initialize(const FPCGDataCollection& InputData, TWeakObjectPtr<UPCGComponent> SourceComponent, const UPCGNode* Node) override;
+protected:
+	virtual bool PrepareDataInternal(FPCGContext* InContext) const override;
+	virtual bool ExecuteInternal(FPCGContext* InContext) const override;
+};
\ No newline at end of file
