Update mesh path to properly typed soft object pointer

#rnx

diff --git a/Engine/Plugins/Experimental/PCGGeometryScriptInterop/Source/PCGGeometryScriptInterop/Private/Elements/PCGMeshSampler.cpp b/Engine/Plugins/Experimental/PCGGeometryScriptInterop/Source/PCGGeometryScriptInterop/Private/Elements/PCGMeshSampler.cpp
--- a/Engine/Plugins/Experimental/PCGGeometryScriptInterop/Source/PCGGeometryScriptInterop/Private/Elements/PCGMeshSampler.cpp
+++ b/Engine/Plugins/Experimental/PCGGeometryScriptInterop/Source/PCGGeometryScriptInterop/Private/Elements/PCGMeshSampler.cpp
@@ -30,6 +30,19 @@ UPCGMeshSamplerSettings::UPCGMeshSamplerSettings()
 	bUseSeed = true;
 }
 
+void UPCGMeshSamplerSettings::PostLoad()
+{
+	Super::PostLoad();
+
+#if WITH_EDITOR
+	if (!StaticMeshPath_DEPRECATED.IsNull())
+	{
+		StaticMesh = StaticMeshPath_DEPRECATED;
+		StaticMeshPath_DEPRECATED.Reset();
+	}
+#endif
+}
+
 TArray<FPCGPinProperties> UPCGMeshSamplerSettings::InputPinProperties() const
 {
 	return TArray<FPCGPinProperties>{};
@@ -91,12 +104,12 @@ bool FPCGMeshSamplerElement::PrepareDataInternal(FPCGContext* InContext) const
 	check(Settings);
 
 	// TODO: Could be async
-	TSoftObjectPtr<UStaticMesh> StaticMeshPtr{ Settings->StaticMeshPath };
+	TSoftObjectPtr<UStaticMesh> StaticMeshPtr = Settings->StaticMesh;
 	UStaticMesh* StaticMesh = StaticMeshPtr.LoadSynchronous();
 
 	if (!StaticMesh)
 	{
-		PCGE_LOG(Error, GraphAndLog, FText::Format(LOCTEXT("MeshDoesNotExist", "Provided static mesh does not exist: '{0}'"), FText::FromString(Settings->StaticMeshPath.ToString())));
+		PCGE_LOG(Error, GraphAndLog, FText::Format(LOCTEXT("MeshDoesNotExist", "Provided static mesh does not exist or could not be loaded: '{0}'"), FText::FromString(StaticMeshPtr.ToString())));
 		return true;
 	}
 
diff --git a/Engine/Plugins/Experimental/PCGGeometryScriptInterop/Source/PCGGeometryScriptInterop/Public/Elements/PCGMeshSampler.h b/Engine/Plugins/Experimental/PCGGeometryScriptInterop/Source/PCGGeometryScriptInterop/Public/Elements/PCGMeshSampler.h
--- a/Engine/Plugins/Experimental/PCGGeometryScriptInterop/Source/PCGGeometryScriptInterop/Public/Elements/PCGMeshSampler.h
+++ b/Engine/Plugins/Experimental/PCGGeometryScriptInterop/Source/PCGGeometryScriptInterop/Public/Elements/PCGMeshSampler.h
@@ -47,6 +47,10 @@ class PCGGEOMETRYSCRIPTINTEROP_API UPCGMeshSamplerSettings : public UPCGSettings
 public:
 	UPCGMeshSamplerSettings();
 
+	//~Begin UObject interface
+	virtual void PostLoad() override;
+	//~End UObject interface
+
 	//~Begin UPCGSettings interface
 #if WITH_EDITOR
 	virtual FName GetDefaultNodeName() const override { return FName(TEXT("MeshSampler")); }
@@ -67,8 +71,8 @@ public:
 	EPCGMeshSamplingMethod SamplingMethod = EPCGMeshSamplingMethod::OnePointPerTriangle;
 
 	/** Soft Object Path to the mesh to sample from. Will be loaded. */
-	UPROPERTY(BlueprintReadWrite, EditAnywhere, Category = "Settings", meta = (PCG_Overridable, AllowedClasses = "/Script/Engine.StaticMesh", DisplayName = "Static Mesh"))
-	FSoftObjectPath StaticMeshPath;
+	UPROPERTY(BlueprintReadWrite, EditAnywhere, Category = "Settings", meta = (PCG_Overridable))
+	TSoftObjectPtr<UStaticMesh> StaticMesh;
 
 	/** In "One Point Per Vertex" option, will assign point density from the red component of the vertex color. */
 	UPROPERTY(BlueprintReadWrite, EditAnywhere, Category = "Settings|Per-Vertex Options", meta = (PCG_Overridable, EditCondition = "SamplingMethod == EPCGMeshSamplingMethod::OnePointPerVertex"))
@@ -99,6 +103,13 @@ public:
 
 	UPROPERTY(BlueprintReadWrite, EditAnywhere, Category = "Settings|Poisson sampling", meta = (PCG_Overridable, EditCondition = "SamplingMethod == EPCGMeshSamplingMethod::PoissonSampling"))
 	FGeometryScriptNonUniformPointSamplingOptions NonUniformSamplingOptions;
+
+protected:
+#if WITH_EDITORONLY_DATA
+	// Deprecated in UE 5.3 in favor of StaticMesh
+	UPROPERTY()
+	FSoftObjectPath StaticMeshPath_DEPRECATED;
+#endif
 };
 
 /**
