Added PCG acronym to plugin names to make them easier to find.

diff --git a/Engine/Plugins/Experimental/PCG/PCG.uplugin b/Engine/Plugins/Experimental/PCG/PCG.uplugin
--- a/Engine/Plugins/Experimental/PCG/PCG.uplugin
+++ b/Engine/Plugins/Experimental/PCG/PCG.uplugin
@@ -2,7 +2,7 @@
 	"FileVersion": 3,
 	"Version": 8,
 	"VersionName": "0.1",
-	"FriendlyName": "Procedural Content Generation Framework",
+	"FriendlyName": "Procedural Content Generation Framework (PCG)",
 	"Description": "Visual scripting framework for procedurally populating worlds with content in editor and/or at run-time.",
 	"Category": "Editor",
 	"CreatedBy": "Epic Games, Inc.",
diff --git a/Engine/Plugins/Experimental/PCGGeometryScriptInterop/PCGGeometryScriptInterop.uplugin b/Engine/Plugins/Experimental/PCGGeometryScriptInterop/PCGGeometryScriptInterop.uplugin
--- a/Engine/Plugins/Experimental/PCGGeometryScriptInterop/PCGGeometryScriptInterop.uplugin
+++ b/Engine/Plugins/Experimental/PCGGeometryScriptInterop/PCGGeometryScriptInterop.uplugin
@@ -2,7 +2,7 @@
 	"FileVersion": 3,
 	"Version": 1,
 	"VersionName": "0.1",
-	"FriendlyName": "Procedural Content Generation Framework Geometry Script Interop",
+	"FriendlyName": "Procedural Content Generation Framework (PCG) Geometry Script Interop",
 	"Description": "Extra plugin for Procedural Content Generation Framework interacting with Geometry Scripts.",
 	"Category": "Editor",
 	"CreatedBy": "Epic Games, Inc.",
