Added missing "steepness" control to sampler nodes:

* Volume Sampler
* Spline Sampler
* Mesh Sampler

diff --git a/Engine/Plugins/Experimental/PCG/Source/PCG/Private/Elements/PCGSplineSampler.cpp b/Engine/Plugins/Experimental/PCG/Source/PCG/Private/Elements/PCGSplineSampler.cpp
--- a/Engine/Plugins/Experimental/PCG/Source/PCG/Private/Elements/PCGSplineSampler.cpp
+++ b/Engine/Plugins/Experimental/PCG/Source/PCG/Private/Elements/PCGSplineSampler.cpp
@@ -532,6 +532,7 @@ namespace PCGSplineSampler
 			FPCGPoint OutPoint;
 			OutPoint.Density = 1.0f;
 			OutPoint.SetLocalBounds(InResult.Box);
+			OutPoint.Steepness = Params.PointSteepness;
 
 			if (ProjectionTargetData)
 			{
@@ -652,6 +653,7 @@ namespace PCGSplineSampler
 					}
 
 					SetSeed(OutPoint, TentativeLocationLS, Params);
+					OutPoint.Steepness = Params.PointSteepness;
 					SetMetadata(InResult, OutPoint, OutPointData->Metadata);
 					OutPointData->GetMutablePoints().Add(OutPoint);
 				}
@@ -1104,6 +1106,7 @@ namespace PCGSplineSampler
 				Point.Density = InteriorPoint.Get<2>();
 				Point.BoundsMin = BoundsMin;
 				Point.BoundsMax = BoundsMax;
+				Point.Steepness = Params.PointSteepness;
 			}
 		}
 	}
diff --git a/Engine/Plugins/Experimental/PCG/Source/PCG/Private/Elements/PCGVolumeSampler.cpp b/Engine/Plugins/Experimental/PCG/Source/PCG/Private/Elements/PCGVolumeSampler.cpp
--- a/Engine/Plugins/Experimental/PCG/Source/PCG/Private/Elements/PCGVolumeSampler.cpp
+++ b/Engine/Plugins/Experimental/PCG/Source/PCG/Private/Elements/PCGVolumeSampler.cpp
@@ -93,7 +93,7 @@ namespace PCGVolumeSampler
 			}
 		}
 
-		FPCGAsync::AsyncPointProcessing(InContext, NumIterations, Points, [InVolume, InBoundingShape, VoxelSize, MinX, MaxX, MinY, MaxY, MinZ, MaxZ](int32 Index, FPCGPoint& OutPoint)
+		FPCGAsync::AsyncPointProcessing(InContext, NumIterations, Points, [InVolume, InBoundingShape, PointSteepness = InSamplerSettings.PointSteepness, VoxelSize, MinX, MaxX, MinY, MaxY, MinZ](int32 Index, FPCGPoint& OutPoint)
 		{
 			const int X = MinX + (Index % (MaxX - MinX));
 			const int Y = MinY + (Index / (MaxX - MinX) % (MaxY - MinY));
@@ -119,6 +119,7 @@ namespace PCGVolumeSampler
 				}
 
 				OutPoint.Seed = PCGHelpers::ComputeSeed(X, Y, Z);
+				OutPoint.Steepness = PointSteepness;
 				return true;
 			}
 			else
@@ -174,6 +175,7 @@ bool FPCGVolumeSamplerElement::ExecuteInternal(FPCGContext* Context) const
 
 	PCGVolumeSampler::FVolumeSamplerSettings SamplerSettings;
 	SamplerSettings.VoxelSize = VoxelSize;
+	SamplerSettings.PointSteepness = Settings->PointSteepness;
 
 	TArray<FPCGTaggedData>& Outputs = Context->OutputData.TaggedData;
 
diff --git a/Engine/Plugins/Experimental/PCG/Source/PCG/Public/Elements/PCGSplineSampler.h b/Engine/Plugins/Experimental/PCG/Source/PCG/Public/Elements/PCGSplineSampler.h
--- a/Engine/Plugins/Experimental/PCG/Source/PCG/Public/Elements/PCGSplineSampler.h
+++ b/Engine/Plugins/Experimental/PCG/Source/PCG/Public/Elements/PCGSplineSampler.h
@@ -110,6 +110,14 @@ struct PCG_API FPCGSplineSamplerParams
 	UPROPERTY(EditAnywhere, Category = Settings, meta = (EditCondition = "Dimension==EPCGSplineSamplingDimension::OnInterior"))
 	FRuntimeFloatCurve InteriorDensityFalloffCurve;
 
+	/** Each PCG point represents a discretized, volumetric region of world space. The points' Steepness value [0.0 to
+	 * 1.0] establishes how "hard" or "soft" that volume will be represented. From 0, it will ramp up linearly
+	 * increasing its influence over the density from the point's center to up to two times the bounds. At 1, it will
+	 * represent a binary box function with the size of the point's bounds.
+	 */
+	UPROPERTY(BlueprintReadWrite, EditAnywhere, Category = "Settings|Points", meta=(ClampMin="0", ClampMax="1", PCG_Overridable))
+	float PointSteepness = 0.5f;
+
 	/** Controls whether we will seed the sampled points using the final world position or the local position */
 	UPROPERTY(BlueprintReadWrite, EditAnywhere, Category = "Settings|Seeding")
 	bool bSeedFromLocalPosition = false;
diff --git a/Engine/Plugins/Experimental/PCG/Source/PCG/Public/Elements/PCGSurfaceSampler.h b/Engine/Plugins/Experimental/PCG/Source/PCG/Public/Elements/PCGSurfaceSampler.h
--- a/Engine/Plugins/Experimental/PCG/Source/PCG/Public/Elements/PCGSurfaceSampler.h
+++ b/Engine/Plugins/Experimental/PCG/Source/PCG/Public/Elements/PCGSurfaceSampler.h
@@ -67,6 +67,11 @@ public:
 	UPROPERTY(BlueprintReadWrite, EditAnywhere, Category = "Settings|Points", meta = (PCG_Overridable))
 	bool bApplyDensityToPoints = true;
 
+	/** Each PCG point represents a discretized, volumetric region of world space. The points' Steepness value [0.0 to
+	 * 1.0] establishes how "hard" or "soft" that volume will be represented. From 0, it will ramp up linearly
+	 * increasing its influence over the density from the point's center to up to two times the bounds. At 1, it will
+	 * represent a binary box function with the size of the point's bounds.
+	 */
 	UPROPERTY(BlueprintReadWrite, EditAnywhere, Category = "Settings|Points", meta=(ClampMin="0", ClampMax="1", PCG_Overridable))
 	float PointSteepness = 0.5f;
 
diff --git a/Engine/Plugins/Experimental/PCG/Source/PCG/Public/Elements/PCGVolumeSampler.h b/Engine/Plugins/Experimental/PCG/Source/PCG/Public/Elements/PCGVolumeSampler.h
--- a/Engine/Plugins/Experimental/PCG/Source/PCG/Public/Elements/PCGVolumeSampler.h
+++ b/Engine/Plugins/Experimental/PCG/Source/PCG/Public/Elements/PCGVolumeSampler.h
@@ -16,6 +16,7 @@ namespace PCGVolumeSampler
 	struct FVolumeSamplerSettings
 	{
 		FVector VoxelSize;
+		float PointSteepness = 0.0f;
 	};
 
 	UPCGPointData* SampleVolume(FPCGContext* InContext, const UPCGSpatialData* InVolume, const FVolumeSamplerSettings& InSamplerSettings);
@@ -40,6 +41,14 @@ public:
 	UPROPERTY(BlueprintReadWrite, EditAnywhere, Category = Settings, meta = (PCG_Overridable))
 	bool bUnbounded = false;
 
+	/** Each PCG point represents a discretized, volumetric region of world space. The points' Steepness value [0.0 to
+	 * 1.0] establishes how "hard" or "soft" that volume will be represented. From 0, it will ramp up linearly
+	 * increasing its influence over the density from the point's center to up to two times the bounds. At 1, it will
+	 * represent a binary box function with the size of the point's bounds.
+	 */
+	UPROPERTY(BlueprintReadWrite, EditAnywhere, Category = "Settings|Points", meta=(ClampMin="0", ClampMax="1", PCG_Overridable))
+	float PointSteepness = 0.5f;
+	
 	//~Begin UPCGSettings interface
 #if WITH_EDITOR
 	virtual FName GetDefaultNodeName() const override { return FName(TEXT("VolumeSampler")); }
diff --git a/Engine/Plugins/Experimental/PCGGeometryScriptInterop/Source/PCGGeometryScriptInterop/Private/Elements/PCGMeshSampler.cpp b/Engine/Plugins/Experimental/PCGGeometryScriptInterop/Source/PCGGeometryScriptInterop/Private/Elements/PCGMeshSampler.cpp
--- a/Engine/Plugins/Experimental/PCGGeometryScriptInterop/Source/PCGGeometryScriptInterop/Private/Elements/PCGMeshSampler.cpp
+++ b/Engine/Plugins/Experimental/PCGGeometryScriptInterop/Source/PCGGeometryScriptInterop/Private/Elements/PCGMeshSampler.cpp
@@ -223,6 +223,7 @@ bool FPCGMeshSamplerElement::ExecuteInternal(FPCGContext* InContext) const
 			OutPoint.Transform = FTransform{ FRotationMatrix::MakeFromZ(Normal).Rotator(), Position, FVector{1.0, 1.0, 1.0} };
 			OutPoint.Density = Settings->bUseRedAsDensity ? Color.R : 1.0f;
 			OutPoint.Color = Color;
+			OutPoint.Steepness = Settings->PointSteepness;
 
 			UPCGBlueprintHelpers::SetSeedFromPosition(OutPoint);
 
@@ -238,7 +239,7 @@ bool FPCGMeshSamplerElement::ExecuteInternal(FPCGContext* InContext) const
 
 		const TArray<int32>& TriangleIds = *Context->TriangleIds.List.Get();
 
-		auto IterationBody = [DynamicMesh = Context->DynamicMesh, &TriangleIds](int32 Index, FPCGPoint& OutPoint) -> bool
+		auto IterationBody = [DynamicMesh = Context->DynamicMesh, PointSteepness = Settings->PointSteepness, &TriangleIds](int32 Index, FPCGPoint& OutPoint) -> bool
 		{
 			const int32 TriangleId = TriangleIds[Index];
 
@@ -253,6 +254,7 @@ bool FPCGMeshSamplerElement::ExecuteInternal(FPCGContext* InContext) const
 			OutPoint = FPCGPoint{};
 			OutPoint.Transform = FTransform{ FRotationMatrix::MakeFromZ(Normal).Rotator(), Position, FVector{1.0, 1.0, 1.0} };
 			OutPoint.Density = 1.0f;
+			OutPoint.Steepness = PointSteepness;
 
 			UPCGBlueprintHelpers::SetSeedFromPosition(OutPoint);
 
@@ -318,6 +320,7 @@ bool FPCGMeshSamplerElement::ExecuteInternal(FPCGContext* InContext) const
 					FPCGPoint& OutPoint = Points.Emplace_GetRef();
 					OutPoint.Transform = Sample.ToTransform();
 					OutPoint.Density = 1.0f;
+					OutPoint.Steepness = Settings->PointSteepness;
 
 					UPCGBlueprintHelpers::SetSeedFromPosition(OutPoint);
 				}
diff --git a/Engine/Plugins/Experimental/PCGGeometryScriptInterop/Source/PCGGeometryScriptInterop/Public/Elements/PCGMeshSampler.h b/Engine/Plugins/Experimental/PCGGeometryScriptInterop/Source/PCGGeometryScriptInterop/Public/Elements/PCGMeshSampler.h
--- a/Engine/Plugins/Experimental/PCGGeometryScriptInterop/Source/PCGGeometryScriptInterop/Public/Elements/PCGMeshSampler.h
+++ b/Engine/Plugins/Experimental/PCGGeometryScriptInterop/Source/PCGGeometryScriptInterop/Public/Elements/PCGMeshSampler.h
@@ -104,6 +104,14 @@ public:
 	UPROPERTY(BlueprintReadWrite, EditAnywhere, Category = "Settings|Poisson sampling", meta = (PCG_Overridable, EditCondition = "SamplingMethod == EPCGMeshSamplingMethod::PoissonSampling"))
 	FGeometryScriptNonUniformPointSamplingOptions NonUniformSamplingOptions;
 
+	/** Each PCG point represents a discretized, volumetric region of world space. The points' Steepness value [0.0 to
+	 * 1.0] establishes how "hard" or "soft" that volume will be represented. From 0, it will ramp up linearly
+	 * increasing its influence over the density from the point's center to up to two times the bounds. At 1, it will
+	 * represent a binary box function with the size of the point's bounds.
+	 */
+	UPROPERTY(BlueprintReadWrite, EditAnywhere, Category = "Settings|Points", meta=(ClampMin="0", ClampMax="1", PCG_Overridable))
+	float PointSteepness = 0.5f;
+	
 protected:
 #if WITH_EDITORONLY_DATA
 	// Deprecated in UE 5.3 in favor of StaticMesh
