[PCG] Improvements on Mesh Sampler

* Colors are now set for Triangles and Poisson Sampling
* Color to density is now available for all modes, and we can chose any color channel (was only red before)
* Add possibility to extract UVs and triangle ids

diff --git a/Engine/Plugins/Experimental/PCGGeometryScriptInterop/Source/PCGGeometryScriptInterop/Private/Elements/PCGMeshSampler.cpp b/Engine/Plugins/Experimental/PCGGeometryScriptInterop/Source/PCGGeometryScriptInterop/Private/Elements/PCGMeshSampler.cpp
--- a/Engine/Plugins/Experimental/PCGGeometryScriptInterop/Source/PCGGeometryScriptInterop/Private/Elements/PCGMeshSampler.cpp
+++ b/Engine/Plugins/Experimental/PCGGeometryScriptInterop/Source/PCGGeometryScriptInterop/Private/Elements/PCGMeshSampler.cpp
@@ -9,6 +9,8 @@
 #include "Helpers/PCGAsync.h"
 #include "Helpers/PCGBlueprintHelpers.h"
 #include "Helpers/PCGHelpers.h"
+#include "Metadata/PCGMetadata.h"
+#include "Metadata/PCGMetadataAttributeTpl.h"
 
 #include "UDynamicMesh.h"
 #include "Engine/AssetManager.h"
@@ -40,6 +42,14 @@ void UPCGMeshSamplerSettings::PostLoad()
 		StaticMesh = StaticMeshPath_DEPRECATED;
 		StaticMeshPath_DEPRECATED.Reset();
 	}
+
+	if (bUseRedAsDensity_DEPRECATED)
+	{
+		// It was only available for one point per vertex before. Keep that.
+		bUseColorChannelAsDensity = (SamplingMethod == EPCGMeshSamplingMethod::OnePointPerVertex);
+		ColorChannelAsDensity = EPCGColorChannel::Red;
+		bUseRedAsDensity_DEPRECATED = false;
+	}
 #endif
 }
 
@@ -48,6 +58,13 @@ TArray<FPCGPinProperties> UPCGMeshSamplerSettings::InputPinProperties() const
 	return TArray<FPCGPinProperties>{};
 }
 
+TArray<FPCGPinProperties> UPCGMeshSamplerSettings::OutputPinProperties() const
+{
+	TArray<FPCGPinProperties> Properties;
+	Properties.Emplace(PCGPinConstants::DefaultOutputLabel, EPCGDataType::Point, /*bInAllowMultipleConnections =*/ false, /*bAllowMultipleData =*/ false);
+	return Properties;
+}
+
 FPCGElementPtr UPCGMeshSamplerSettings::CreateElement() const
 {
 	return MakeShared<FPCGMeshSamplerElement>();
@@ -175,6 +192,29 @@ bool FPCGMeshSamplerElement::PrepareDataInternal(FPCGContext* InContext) const
 
 	TArray<FPCGTaggedData>& Outputs = InContext->OutputData.TaggedData;
 	Context->OutPointData = NewObject<UPCGPointData>();
+
+	// It's not clear how to compute UVs for Vertices as they are part of multiple triangles. So disable for this mode. Same for triangle ids.
+	if (Settings->SamplingMethod != EPCGMeshSamplingMethod::OnePointPerVertex)
+	{
+		if (Settings->bExtractUVAsAttribute)
+		{
+			Context->UVAttribute = Context->OutPointData->Metadata->CreateAttribute<FVector2D>(Settings->UVAttributeName, FVector2D::ZeroVector, /*bAllowsInterpolation=*/true, /*bOverrideParent=*/true);
+			if (!Context->UVAttribute)
+			{
+				PCGE_LOG(Warning, GraphAndLog, FText::Format(LOCTEXT("AttributeUVFailed", "Failed to create attribute {0} for UVs. UVs won't be computed"), FText::FromName(Settings->UVAttributeName)));
+			}
+		}
+
+		if (Settings->bOutputTriangleIds)
+		{
+			Context->TriangleIdAttribute = Context->OutPointData->Metadata->CreateAttribute<int32>(Settings->TriangleIdAttributeName, 0, /*bAllowsInterpolation=*/false, /*bOverrideParent=*/true);
+			if (!Context->TriangleIdAttribute)
+			{
+				PCGE_LOG(Warning, GraphAndLog, FText::Format(LOCTEXT("AttributeTriangleIdFailed", "Failed to create attribute {0} for triangles ids. Triangle Ids won't be output"), FText::FromName(Settings->TriangleIdAttributeName)));
+			}
+		}
+	}
+
 	Outputs.Emplace_GetRef().Data = Context->OutPointData;
 
 	Context->bDataPrepared = true;
@@ -202,6 +242,75 @@ bool FPCGMeshSamplerElement::ExecuteInternal(FPCGContext* InContext) const
 	const bool bEnableTimeSlicing = true;
 	const int Seed = Context->GetSeed();
 
+	auto SetUVValueAndTriangleId = [Context, Settings](int32 TriangleId, const FVector& BarycentricCoord, FPCGPoint& OutPoint)
+	{
+		if (Context->UVAttribute)
+		{
+			bool bHasValidUVs = false;
+			FVector2D InterpolatedUV{};
+			UGeometryScriptLibrary_MeshQueryFunctions::GetInterpolatedTriangleUV(Context->DynamicMesh, /*UVSetIndex=*/ Settings->UVChannel, TriangleId, BarycentricCoord, bHasValidUVs, InterpolatedUV);
+			if (bHasValidUVs)
+			{
+				Context->OutPointData->Metadata->InitializeOnSet(OutPoint.MetadataEntry);
+				Context->UVAttribute->SetValue(OutPoint.MetadataEntry, InterpolatedUV);
+			}
+		}
+
+		if (Context->TriangleIdAttribute)
+		{
+			Context->OutPointData->Metadata->InitializeOnSet(OutPoint.MetadataEntry);
+			Context->TriangleIdAttribute->SetValue(OutPoint.MetadataEntry, TriangleId);
+		}
+	};
+
+	// Preparing the set function to extract the color to density.
+	auto SetPointDensityTo1 = [](const FLinearColor&, FPCGPoint& OutPoint){ OutPoint.Density = 1.0f; };
+	auto SetPointDensityToRed = [](const FLinearColor& Color, FPCGPoint& OutPoint){ OutPoint.Density = Color.R; };
+	auto SetPointDensityToGreen = [](const FLinearColor& Color, FPCGPoint& OutPoint){ OutPoint.Density = Color.G; };
+	auto SetPointDensityToBlue = [](const FLinearColor& Color, FPCGPoint& OutPoint){ OutPoint.Density = Color.B; };
+	auto SetPointDensityToAlpha = [](const FLinearColor& Color, FPCGPoint& OutPoint){ OutPoint.Density = Color.A; };
+
+	// Store it in a function pointer.
+	void(*SetPointDensityPtr)(const FLinearColor&, FPCGPoint&) = SetPointDensityTo1;
+
+	if (Settings->bUseColorChannelAsDensity)
+	{
+		switch (Settings->ColorChannelAsDensity)
+		{
+		case EPCGColorChannel::Red:
+			SetPointDensityPtr = SetPointDensityToRed;
+			break;
+		case EPCGColorChannel::Green:
+			SetPointDensityPtr = SetPointDensityToGreen;
+			break;
+		case EPCGColorChannel::Blue:
+			SetPointDensityPtr = SetPointDensityToBlue;
+			break;
+		case EPCGColorChannel::Alpha:
+			SetPointDensityPtr = SetPointDensityToAlpha;
+			break;
+		default:
+			checkNoEntry();
+			break;
+		}
+	}
+
+	auto SetPointColorAndDensity = [Context, Settings, SetPointDensityPtr](int32 TriangleId, const FVector& BarycentricCoord, FPCGPoint& OutPoint)
+	{
+		FLinearColor Color;
+		bool bValidVertexColor = false;
+		UGeometryScriptLibrary_MeshQueryFunctions::GetInterpolatedTriangleVertexColor(Context->DynamicMesh, TriangleId, BarycentricCoord, FLinearColor::White, bValidVertexColor, Color);
+		if (bValidVertexColor)
+		{
+			OutPoint.Color = Color;
+			SetPointDensityPtr(Color, OutPoint);
+		}
+		else
+		{
+			OutPoint.Density = 1.0f;
+		}
+	};
+
 	switch (Settings->SamplingMethod)
 	{
 	case EPCGMeshSamplingMethod::OnePointPerVertex:
@@ -212,7 +321,7 @@ bool FPCGMeshSamplerElement::ExecuteInternal(FPCGContext* InContext) const
 		const TArray<FLinearColor>& Colors = *Context->Colors.List.Get();
 		const TArray<FVector>& Normals = *Context->Normals.List.Get();
 
-		auto IterationBody = [&Positions, &Colors, &Normals, Settings](int32 Index, FPCGPoint& OutPoint) -> bool
+		auto IterationBody = [&Positions, &Colors, &Normals, Settings, Context, SetPointDensityPtr](int32 Index, FPCGPoint& OutPoint) -> bool
 		{
 			const FVector& Position = Positions[Index];
 			const FLinearColor& Color = Colors[Index];
@@ -220,10 +329,11 @@ bool FPCGMeshSamplerElement::ExecuteInternal(FPCGContext* InContext) const
 
 			OutPoint = FPCGPoint{};
 			OutPoint.Transform = FTransform{ FRotationMatrix::MakeFromZ(Normal).Rotator(), Position, FVector{1.0, 1.0, 1.0} };
-			OutPoint.Density = Settings->bUseRedAsDensity ? Color.R : 1.0f;
 			OutPoint.Color = Color;
 			OutPoint.Steepness = Settings->PointSteepness;
 
+			SetPointDensityPtr(Color, OutPoint);
+
 			UPCGBlueprintHelpers::SetSeedFromPosition(OutPoint);
 
 			return true;
@@ -238,7 +348,7 @@ bool FPCGMeshSamplerElement::ExecuteInternal(FPCGContext* InContext) const
 
 		const TArray<int32>& TriangleIds = *Context->TriangleIds.List.Get();
 
-		auto IterationBody = [DynamicMesh = Context->DynamicMesh, PointSteepness = Settings->PointSteepness, &TriangleIds](int32 Index, FPCGPoint& OutPoint) -> bool
+		auto IterationBody = [DynamicMesh = Context->DynamicMesh, PointSteepness = Settings->PointSteepness, &TriangleIds, SetUVValueAndTriangleId, SetPointColorAndDensity](int32 Index, FPCGPoint& OutPoint) -> bool
 		{
 			const int32 TriangleId = TriangleIds[Index];
 
@@ -247,14 +357,20 @@ bool FPCGMeshSamplerElement::ExecuteInternal(FPCGContext* InContext) const
 
 			UGeometryScriptLibrary_MeshQueryFunctions::GetTrianglePositions(DynamicMesh, TriangleId, bIsValidTriangle, Vertex1, Vertex2, Vertex3);
 			const FVector Normal = UGeometryScriptLibrary_MeshQueryFunctions::GetTriangleFaceNormal(DynamicMesh, TriangleId, bIsValidTriangle);
-
 			const FVector Position = (Vertex1 + Vertex2 + Vertex3) / 3.0;
 
 			OutPoint = FPCGPoint{};
 			OutPoint.Transform = FTransform{ FRotationMatrix::MakeFromZ(Normal).Rotator(), Position, FVector{1.0, 1.0, 1.0} };
-			OutPoint.Density = 1.0f;
 			OutPoint.Steepness = PointSteepness;
 
+			FVector Dummy1, Dummy2, Dummy3;
+			FVector BarycentricCoord;
+			bool bIsValid = false;
+			UGeometryScriptLibrary_MeshQueryFunctions::ComputeTriangleBarycentricCoords(DynamicMesh, TriangleId, bIsValid, Position, Dummy1, Dummy2, Dummy3, BarycentricCoord);
+			
+			SetPointColorAndDensity(TriangleId, BarycentricCoord, OutPoint);
+			SetUVValueAndTriangleId(TriangleId, BarycentricCoord, OutPoint);
+
 			UPCGBlueprintHelpers::SetSeedFromPosition(OutPoint);
 
 			return true;
@@ -275,7 +391,7 @@ bool FPCGMeshSamplerElement::ExecuteInternal(FPCGContext* InContext) const
 			Context->SamplingProgess = MakeUnique<FProgressCancel>();
 			Context->SamplingProgess->CancelF = [Context]() -> bool { return Context->StopSampling; };
 
-			auto SamplingFuture = [Settings, Context, Seed]() -> bool
+			auto SamplingFuture = [Settings, Context, Seed, SetUVValueAndTriangleId, SetPointColorAndDensity]() -> bool
 			{
 				TRACE_CPUPROFILER_EVENT_SCOPE(FPCGMeshSamplerElement::Execute::PoissonSampling);
 
@@ -292,6 +408,8 @@ bool FPCGMeshSamplerElement::ExecuteInternal(FPCGContext* InContext) const
 					PointSampling.SizeDistributionPower = FMath::Clamp(Settings->NonUniformSamplingOptions.SizeDistributionPower, 1.0, 10.0);
 				}
 
+				PointSampling.bComputeBarycentrics = true;
+
 				PointSampling.ComputePoissonSampling(Context->DynamicMesh->GetMeshRef(), Context->SamplingProgess.Get());
 
 				if (Context->StopSampling)
@@ -303,8 +421,9 @@ bool FPCGMeshSamplerElement::ExecuteInternal(FPCGContext* InContext) const
 				Points.Reserve(PointSampling.Samples.Num());
 
 				int Count = 0;
-				for (UE::Geometry::FFrame3d& Sample : PointSampling.Samples)
+				for (int32 i = 0; i < PointSampling.Samples.Num(); ++i)
 				{
+					UE::Geometry::FFrame3d& Sample = PointSampling.Samples[i];
 					// Avoid to check too many times
 					constexpr int CancelledCheckNum = 25;
 					if (++Count == CancelledCheckNum)
@@ -316,11 +435,16 @@ bool FPCGMeshSamplerElement::ExecuteInternal(FPCGContext* InContext) const
 						}
 					}
 
+					const int32 TriangleId = PointSampling.TriangleIDs[i];
+					const FVector BarycentricCoords = PointSampling.BarycentricCoords[i];
+
 					FPCGPoint& OutPoint = Points.Emplace_GetRef();
 					OutPoint.Transform = Sample.ToTransform();
-					OutPoint.Density = 1.0f;
 					OutPoint.Steepness = Settings->PointSteepness;
 
+					SetPointColorAndDensity(TriangleId, BarycentricCoords, OutPoint);
+					SetUVValueAndTriangleId(TriangleId, BarycentricCoords, OutPoint);
+
 					UPCGBlueprintHelpers::SetSeedFromPosition(OutPoint);
 				}
 
diff --git a/Engine/Plugins/Experimental/PCGGeometryScriptInterop/Source/PCGGeometryScriptInterop/Public/Elements/PCGMeshSampler.h b/Engine/Plugins/Experimental/PCGGeometryScriptInterop/Source/PCGGeometryScriptInterop/Public/Elements/PCGMeshSampler.h
--- a/Engine/Plugins/Experimental/PCGGeometryScriptInterop/Source/PCGGeometryScriptInterop/Public/Elements/PCGMeshSampler.h
+++ b/Engine/Plugins/Experimental/PCGGeometryScriptInterop/Source/PCGGeometryScriptInterop/Public/Elements/PCGMeshSampler.h
@@ -17,6 +17,9 @@
 
 #include "PCGMeshSampler.generated.h"
 
+template <typename T>
+class FPCGMetadataAttribute;
+
 class FProgressCancel;
 class UDynamicMesh;
 class UPCGPointData;
@@ -35,6 +38,15 @@ enum class EPCGMeshSamplingMethod
 	PoissonSampling
 };
 
+UENUM()
+enum class EPCGColorChannel
+{
+	Red,
+	Green,
+	Blue,
+	Alpha
+};
+
 /**
 * Sample points on a mesh
 */
@@ -59,7 +71,7 @@ public:
 #endif
 
 	virtual TArray<FPCGPinProperties> InputPinProperties() const override;
-	virtual TArray<FPCGPinProperties> OutputPinProperties() const override { return DefaultPointOutputPinProperties(); }
+	virtual TArray<FPCGPinProperties> OutputPinProperties() const override;
 
 protected:
 	virtual FPCGElementPtr CreateElement() const override;
@@ -73,9 +85,12 @@ public:
 	UPROPERTY(BlueprintReadWrite, EditAnywhere, Category = "Settings", meta = (PCG_Overridable))
 	TSoftObjectPtr<UStaticMesh> StaticMesh;
 
-	/** In "One Point Per Vertex" option, will assign point density from the red component of the vertex color. */
-	UPROPERTY(BlueprintReadWrite, EditAnywhere, Category = "Settings|Per-Vertex Options", meta = (PCG_Overridable, EditCondition = "SamplingMethod == EPCGMeshSamplingMethod::OnePointPerVertex"))
-	bool bUseRedAsDensity = false;
+	/** Will extract the color channel into the density. */
+	UPROPERTY(BlueprintReadWrite, EditAnywhere, Category = "Settings|Color & Density", meta = (PCG_Overridable, InlineEditConditionToggle, PCG_OverrideAliases="bUseRedAsDensity"))
+	bool bUseColorChannelAsDensity = false;
+	
+	UPROPERTY(BlueprintReadWrite, EditAnywhere, Category = "Settings|Color & Density", meta = (PCG_Overridable, EditCondition="bUseColorChannelAsDensity"))
+	EPCGColorChannel ColorChannelAsDensity = EPCGColorChannel::Red;
 
 	/** Enable voxelisation as a preparation pass. Can be more expensive given the VoxelSize. */
 	UPROPERTY(BlueprintReadWrite, EditAnywhere, Category = "Settings|Voxelize Options", meta = (PCG_Overridable))
@@ -103,6 +118,21 @@ public:
 	UPROPERTY(BlueprintReadWrite, EditAnywhere, Category = "Settings|Poisson sampling", meta = (PCG_Overridable, EditCondition = "SamplingMethod == EPCGMeshSamplingMethod::PoissonSampling"))
 	FGeometryScriptNonUniformPointSamplingOptions NonUniformSamplingOptions;
 
+	UPROPERTY(BlueprintReadWrite, EditAnywhere, Category = "Settings|UVs", meta = (PCG_Overridable, EditCondition = "SamplingMethod != EPCGMeshSamplingMethod::OnePointPerVertex"))
+	bool bExtractUVAsAttribute = false;
+
+	UPROPERTY(BlueprintReadWrite, EditAnywhere, Category = "Settings|UVs", meta = (PCG_Overridable, EditCondition = "SamplingMethod != EPCGMeshSamplingMethod::OnePointPerVertex && bExtractUVAsAttribute", EditConditionHides))
+	FName UVAttributeName = TEXT("UV");
+
+	UPROPERTY(BlueprintReadWrite, EditAnywhere, Category = "Settings|UVs", meta = (PCG_Overridable, EditCondition = "SamplingMethod != EPCGMeshSamplingMethod::OnePointPerVertex && bExtractUVAsAttribute", EditConditionHides))
+	int32 UVChannel = 0;
+
+	UPROPERTY(BlueprintReadWrite, EditAnywhere, Category = "Settings|Extra", meta = (PCG_Overridable, EditCondition = "SamplingMethod != EPCGMeshSamplingMethod::OnePointPerVertex"))
+	bool bOutputTriangleIds = false;
+
+	UPROPERTY(BlueprintReadWrite, EditAnywhere, Category = "Settings|Extra", meta = (PCG_Overridable, EditCondition = "SamplingMethod != EPCGMeshSamplingMethod::OnePointPerVertex && bOutputTriangleIds", EditConditionHides))
+	FName TriangleIdAttributeName = TEXT("TriangleId");
+
 	/** Each PCG point represents a discretized, volumetric region of world space. The points' Steepness value [0.0 to
 	 * 1.0] establishes how "hard" or "soft" that volume will be represented. From 0, it will ramp up linearly
 	 * increasing its influence over the density from the point's center to up to two times the bounds. At 1, it will
@@ -116,6 +146,9 @@ protected:
 	// Deprecated in UE 5.3 in favor of StaticMesh
 	UPROPERTY()
 	FSoftObjectPath StaticMeshPath_DEPRECATED;
+
+	UPROPERTY()
+	bool bUseRedAsDensity_DEPRECATED = false;
 #endif
 };
 
@@ -138,6 +171,10 @@ struct FPCGMeshSamplerContext : public FPCGContext
 	// Output point data
 	UPCGPointData* OutPointData = nullptr;
 
+	// Optional attributes.
+	FPCGMetadataAttribute<FVector2D>* UVAttribute = nullptr;
+	FPCGMetadataAttribute<int32>* TriangleIdAttribute = nullptr;
+
 	// For Poisson sampling, we are starting a future that is not framebound
 	// Store the future and synchronisation items in the context
 	TFuture<bool> SamplingFuture;
