Fixed missing enum size declaration due to bug in UHT

#rnx

diff --git a/Engine/Plugins/Experimental/PCGGeometryScriptInterop/Source/PCGGeometryScriptInterop/Public/Elements/PCGMeshSampler.h b/Engine/Plugins/Experimental/PCGGeometryScriptInterop/Source/PCGGeometryScriptInterop/Public/Elements/PCGMeshSampler.h
--- a/Engine/Plugins/Experimental/PCGGeometryScriptInterop/Source/PCGGeometryScriptInterop/Public/Elements/PCGMeshSampler.h
+++ b/Engine/Plugins/Experimental/PCGGeometryScriptInterop/Source/PCGGeometryScriptInterop/Public/Elements/PCGMeshSampler.h
@@ -26,7 +26,7 @@ class UPCGPointData;
 class UStaticMesh;
 
 UENUM(BlueprintType)
-enum class EPCGMeshSamplingMethod
+enum class EPCGMeshSamplingMethod : uint8
 {
 	/** Sample one point (at the center) of each triangle of the mesh. */
 	OnePointPerTriangle,
