[PCG] Mark PCG plugins as beta and remove Geometry Script dependency in the main plugin

#rnx

diff --git a/Engine/Plugins/PCG/PCG.uplugin b/Engine/Plugins/PCG/PCG.uplugin
--- a/Engine/Plugins/PCG/PCG.uplugin
+++ b/Engine/Plugins/PCG/PCG.uplugin
@@ -1,18 +1,18 @@
 {
 	"FileVersion": 3,
 	"Version": 8,
-	"VersionName": "0.1",
+	"VersionName": "0.2",
 	"FriendlyName": "Procedural Content Generation Framework (PCG)",
 	"Description": "Visual scripting framework for procedurally populating worlds with content in editor and/or at run-time.",
 	"Category": "Editor",
 	"CreatedBy": "Epic Games, Inc.",
 	"CreatedByURL": "https://epicgames.com",
-	"DocsURL": "",
+	"DocsURL": "https://docs.unrealengine.com/latest/en-US/procedural-content-generation--framework-in-unreal-engine/",
 	"MarketplaceURL": "",
 	"SupportURL": "",
 	"EnabledByDefault": false,
 	"CanContainContent": true,
-	"IsBetaVersion": false,
+	"IsBetaVersion": true,
 	"Installed": false,
 	"Modules": [
 		{
@@ -32,10 +32,6 @@
 		}
 	],
 	"Plugins": [
-		{
-			"Name": "GeometryScripting",
-			"Enabled": true
-		},
 		{
 			"Name": "EditorScriptingUtilities",
 			"Enabled": true
diff --git a/Engine/Plugins/PCGExternalDataInterop/PCGExternalDataInterop.uplugin b/Engine/Plugins/PCGExternalDataInterop/PCGExternalDataInterop.uplugin
--- a/Engine/Plugins/PCGExternalDataInterop/PCGExternalDataInterop.uplugin
+++ b/Engine/Plugins/PCGExternalDataInterop/PCGExternalDataInterop.uplugin
@@ -1,18 +1,18 @@
 {
 	"FileVersion": 3,
 	"Version": 1,
-	"VersionName": "0.1",
+	"VersionName": "0.2",
 	"FriendlyName": "Procedural Content Generation Framework (PCG) External Data Interop",
 	"Description": "Extra plugin for Procedural Content Generation Framework interacting with external data formats.",
 	"Category": "Editor",
 	"CreatedBy": "Epic Games, Inc.",
 	"CreatedByURL": "https://epicgames.com",
-	"DocsURL": "",
+	"DocsURL": "https://docs.unrealengine.com/latest/en-US/procedural-content-generation--framework-in-unreal-engine/",
 	"MarketplaceURL": "",
 	"SupportURL": "",
 	"EnabledByDefault": false,
 	"CanContainContent": true,
-	"IsBetaVersion": false,
+	"IsBetaVersion": true,
 	"Installed": false,
 	"Modules": [
 		{
diff --git a/Engine/Plugins/PCGGeometryScriptInterop/PCGGeometryScriptInterop.uplugin b/Engine/Plugins/PCGGeometryScriptInterop/PCGGeometryScriptInterop.uplugin
--- a/Engine/Plugins/PCGGeometryScriptInterop/PCGGeometryScriptInterop.uplugin
+++ b/Engine/Plugins/PCGGeometryScriptInterop/PCGGeometryScriptInterop.uplugin
@@ -1,18 +1,18 @@
 {
 	"FileVersion": 3,
 	"Version": 1,
-	"VersionName": "0.1",
+	"VersionName": "0.2",
 	"FriendlyName": "Procedural Content Generation Framework (PCG) Geometry Script Interop",
 	"Description": "Extra plugin for Procedural Content Generation Framework interacting with Geometry Scripts.",
 	"Category": "Editor",
 	"CreatedBy": "Epic Games, Inc.",
 	"CreatedByURL": "https://epicgames.com",
-	"DocsURL": "",
+	"DocsURL": "https://docs.unrealengine.com/latest/en-US/procedural-content-generation--framework-in-unreal-engine/",
 	"MarketplaceURL": "",
 	"SupportURL": "",
 	"EnabledByDefault": false,
 	"CanContainContent": true,
-	"IsBetaVersion": false,
+	"IsBetaVersion": true,
 	"Installed": false,
 	"Modules": [
 		{
