[PCG] Add support for sampling meshes from references

* Will support the sampling of multiple meshes, or sampling on meshes coming from an Actor/Scene Component
* Support of instance vertex colors
* First version is a naive version that does no partitioning, and follows the same input/output relation of Get Property From Point Data, which is 1 output point data per entry.
* Also expose missing function

#rnx

diff --git a/Engine/Plugins/PCG/Source/PCG/Public/Elements/PCGLoadObjectsContext.h b/Engine/Plugins/PCG/Source/PCG/Public/Elements/PCGLoadObjectsContext.h
--- a/Engine/Plugins/PCG/Source/PCG/Public/Elements/PCGLoadObjectsContext.h
+++ b/Engine/Plugins/PCG/Source/PCG/Public/Elements/PCGLoadObjectsContext.h
@@ -8,7 +8,7 @@
 
 struct FPCGLoadObjectsFromPathContext : public FPCGContext, public IPCGAsyncLoadingContext
 {
-	bool InitializeAndRequestLoad(
+	PCG_API bool InitializeAndRequestLoad(
 		FName InputPinName,
 		const FPCGAttributePropertyInputSelector& InputAttributeSelector,
 		const TArray<FSoftObjectPath>& StaticObjectPaths,
diff --git a/Engine/Plugins/PCGGeometryScriptInterop/Source/PCGGeometryScriptInterop/Private/Elements/PCGMeshSampler.cpp b/Engine/Plugins/PCGGeometryScriptInterop/Source/PCGGeometryScriptInterop/Private/Elements/PCGMeshSampler.cpp
--- a/Engine/Plugins/PCGGeometryScriptInterop/Source/PCGGeometryScriptInterop/Private/Elements/PCGMeshSampler.cpp
+++ b/Engine/Plugins/PCGGeometryScriptInterop/Source/PCGGeometryScriptInterop/Private/Elements/PCGMeshSampler.cpp
@@ -22,10 +22,240 @@
 #include "GeometryScript/MeshRepairFunctions.h"
 #include "GeometryScript/MeshVertexColorFunctions.h"
 #include "GeometryScript/MeshVoxelFunctions.h"
+#include "GeometryScript/SceneUtilityFunctions.h"
 #include "Sampling/MeshSurfacePointSampling.h"
 
 #define LOCTEXT_NAMESPACE "PCGMeshSampler"
 
+namespace PCGMeshSampler
+{
+	bool SampleOnePointPerVertex(const UPCGMeshSamplerSettings* Settings, FPCGMeshSamplerContext* Context, FPCGMeshSamplerContext::SetPointDensityFunc SetPointDensityPtr)
+	{
+		check(Settings && Context && SetPointDensityPtr);
+
+		TRACE_CPUPROFILER_EVENT_SCOPE(PCGMeshSampler::SampleOnePointPerVertex);
+
+		auto IterationBody = [Settings, Context, SetPointDensityPtr](int32 ReadIndex, int32 WriteIndex) -> bool
+		{
+			int32 DataIndex = 0;
+			while (!Context->DynamicMeshes[DataIndex] || ReadIndex > Context->StartingIndices[DataIndex])
+			{
+				DataIndex++;
+			}
+
+			const TArray<FVector>& Positions = *Context->Positions[DataIndex].List.Get();
+			const TArray<FLinearColor>& Colors = *Context->Colors[DataIndex].List.Get();
+			const TArray<FVector>& Normals = *Context->Normals[DataIndex].List.Get();
+
+			const int32 CurrentIndex = ReadIndex - Context->StartingIndices[DataIndex];
+
+			const FVector& Position = Positions[CurrentIndex];
+			const FLinearColor& Color = Colors[CurrentIndex];
+			const FVector& Normal = Normals[CurrentIndex];
+
+			FPCGPoint OutPoint{};
+			OutPoint.Transform = FTransform{ FRotationMatrix::MakeFromZ(Normal).Rotator(), Position, FVector::OneVector };
+			OutPoint.Color = Color;
+			OutPoint.Steepness = Settings->PointSteepness;
+
+			SetPointDensityPtr(Color, OutPoint);
+
+			UPCGBlueprintHelpers::SetSeedFromPosition(OutPoint);
+
+			Context->OutPointData[DataIndex]->GetMutablePoints()[CurrentIndex] = OutPoint;
+
+			return true;
+		};
+
+		return FPCGAsync::AsyncProcessingOneToOneEx(&Context->AsyncState, Context->StartingIndices.Last(), []() {}, IterationBody, /*bEnableTimeSlicing=*/true);
+	}
+
+	bool SampleOnePointPerTriangle(const UPCGMeshSamplerSettings* Settings, FPCGMeshSamplerContext* Context, FPCGMeshSamplerContext::SetPointDensityFunc SetPointDensityPtr)
+	{
+		check(Settings && Context && SetPointDensityPtr);
+
+		TRACE_CPUPROFILER_EVENT_SCOPE(FPCGMeshSamplerElement::Execute::OnePointPerTriangle);
+
+		auto IterationBody = [Context, Settings, SetPointDensityPtr](int32 ReadIndex, int32 WriteIndex) -> bool
+		{
+			int32 DataIndex = 0;
+			while (!Context->DynamicMeshes[DataIndex] || ReadIndex >= Context->StartingIndices[DataIndex])
+			{
+				DataIndex++;
+			}
+
+			const TArray<int32>& TriangleIds = *Context->TriangleIds[DataIndex].List.Get();
+			const int32 CurrentIndex = ReadIndex - Context->StartingIndices[DataIndex];
+			const int32 TriangleId = TriangleIds[CurrentIndex];
+
+			FVector Vertex1, Vertex2, Vertex3;
+			bool bIsValidTriangle;
+
+			UGeometryScriptLibrary_MeshQueryFunctions::GetTrianglePositions(Context->DynamicMeshes[DataIndex], TriangleId, bIsValidTriangle, Vertex1, Vertex2, Vertex3);
+			const FVector Normal = UGeometryScriptLibrary_MeshQueryFunctions::GetTriangleFaceNormal(Context->DynamicMeshes[DataIndex], TriangleId, bIsValidTriangle);
+			const FVector Position = (Vertex1 + Vertex2 + Vertex3) / 3.0;
+
+			FPCGPoint OutPoint{};
+			OutPoint.Transform = FTransform{ FRotationMatrix::MakeFromZ(Normal).Rotator(), Position, FVector::OneVector };
+			OutPoint.Steepness = Settings->PointSteepness;
+
+			FVector Dummy1, Dummy2, Dummy3;
+			FVector BarycentricCoord;
+			bool bIsValid = false;
+			UGeometryScriptLibrary_MeshQueryFunctions::ComputeTriangleBarycentricCoords(Context->DynamicMeshes[DataIndex], TriangleId, bIsValid, Position, Dummy1, Dummy2, Dummy3, BarycentricCoord);
+
+			Context->SetPointColorAndDensity(SetPointDensityPtr, TriangleId, BarycentricCoord, OutPoint, DataIndex);
+			Context->SetUVValueAndTriangleId(Settings->UVChannel, TriangleId, BarycentricCoord, OutPoint, DataIndex);
+
+			UPCGBlueprintHelpers::SetSeedFromPosition(OutPoint);
+
+			Context->OutPointData[DataIndex]->GetMutablePoints()[CurrentIndex] = OutPoint;
+
+			return true;
+		};
+
+		return FPCGAsync::AsyncProcessingOneToOneEx(&Context->AsyncState, Context->StartingIndices.Last(), []() {}, IterationBody, /*bEnableTimeSlicing=*/true);
+	}
+
+	bool PoissonSampling(const UPCGMeshSamplerSettings* Settings, FPCGMeshSamplerContext* Context, FPCGMeshSamplerContext::SetPointDensityFunc SetPointDensityPtr)
+	{
+		check(Settings && Context && SetPointDensityPtr);
+
+		// For Poisson sampling, we are calling a "all-in-one" function, where we don't have control for timeslicing.
+		// Since Poisson sampling can be expensive (depending on the radius used), we will do the sampling in a future, put this task to sleep,
+		// and wait for the sampling to wake us up.
+		if (Context->SamplingFutures.IsEmpty())
+		{
+			// Put the task asleep
+			Context->bIsPaused = true;
+			Context->SamplingProgess = MakeUnique<FProgressCancel>();
+			Context->SamplingProgess->CancelF = [Context]() -> bool { return Context->StopSampling; };
+
+			for (int32 DataIndex = 0; DataIndex < Context->DynamicMeshes.Num(); ++DataIndex)
+			{
+				if (!Context->DynamicMeshes[DataIndex])
+				{
+					continue;
+				}
+
+				auto SamplingFuture = [Settings, Context, Seed = Context->GetSeed(), SetPointDensityPtr, DataIndex]() -> bool
+				{
+					TRACE_CPUPROFILER_EVENT_SCOPE(FPCGMeshSamplerElement::Execute::PoissonSampling);
+
+					UE::Geometry::FMeshSurfacePointSampling PointSampling;
+					PointSampling.SampleRadius = Settings->SamplingOptions.SamplingRadius;
+					PointSampling.MaxSamples = Settings->SamplingOptions.MaxNumSamples;
+					PointSampling.RandomSeed = PCGHelpers::ComputeSeed(Seed, Settings->SamplingOptions.RandomSeed);
+					PointSampling.SubSampleDensity = Settings->SamplingOptions.SubSampleDensity;
+
+					if (Settings->NonUniformSamplingOptions.MaxSamplingRadius > PointSampling.SampleRadius)
+					{
+						PointSampling.MaxSampleRadius = Settings->NonUniformSamplingOptions.MaxSamplingRadius;
+						PointSampling.SizeDistribution = UE::Geometry::FMeshSurfacePointSampling::ESizeDistribution(static_cast<int>(Settings->NonUniformSamplingOptions.SizeDistribution));
+						PointSampling.SizeDistributionPower = FMath::Clamp(Settings->NonUniformSamplingOptions.SizeDistributionPower, 1.0, 10.0);
+					}
+
+					PointSampling.bComputeBarycentrics = true;
+
+					PointSampling.ComputePoissonSampling(Context->DynamicMeshes[DataIndex]->GetMeshRef(), Context->SamplingProgess.Get());
+
+					if (Context->StopSampling)
+					{
+						return true;
+					}
+
+					TArray<FPCGPoint>& Points = Context->OutPointData[DataIndex]->GetMutablePoints();
+					Points.Reserve(PointSampling.Samples.Num());
+
+					int Count = 0;
+					for (int32 i = 0; i < PointSampling.Samples.Num(); ++i)
+					{
+						UE::Geometry::FFrame3d& Sample = PointSampling.Samples[i];
+						// Avoid to check too many times
+						constexpr int CancelledCheckNum = 25;
+						if (++Count == CancelledCheckNum)
+						{
+							Count = 0;
+							if (Context->StopSampling)
+							{
+								return true;
+							}
+						}
+
+						const int32 TriangleId = PointSampling.TriangleIDs[i];
+						const FVector BarycentricCoords = PointSampling.BarycentricCoords[i];
+
+						FPCGPoint& OutPoint = Points.Emplace_GetRef();
+						OutPoint.Transform = Sample.ToTransform();
+						OutPoint.Steepness = Settings->PointSteepness;
+
+						Context->SetPointColorAndDensity(SetPointDensityPtr, TriangleId, BarycentricCoords, OutPoint, DataIndex);
+						Context->SetUVValueAndTriangleId(Settings->UVChannel, TriangleId, BarycentricCoords, OutPoint, DataIndex);
+
+						UPCGBlueprintHelpers::SetSeedFromPosition(OutPoint);
+					}
+
+					// Unpause the task
+					Context->bIsPaused = false;
+					return true;
+				};
+
+				Context->SamplingFutures.Add(Async(EAsyncExecution::ThreadPool, std::move(SamplingFuture)));
+			}
+		}
+
+		const bool bIsDone = Algo::AllOf(Context->SamplingFutures, [](const TFuture<bool>& Future) { return Future.IsReady();});
+		if (bIsDone)
+		{
+			for (TFuture<bool>& Future : Context->SamplingFutures)
+			{
+				Future.Reset();
+			}
+		}
+
+		return bIsDone;
+	}
+}
+
+void FPCGMeshSamplerContext::SetUVValueAndTriangleId(int32 UVChannel, int32 TriangleId, const FVector& BarycentricCoord, FPCGPoint& OutPoint, int32 DataIndex)
+{
+	if (UVAttributes.IsValidIndex(DataIndex) && UVAttributes[DataIndex])
+	{
+		bool bHasValidUVs = false;
+		FVector2D InterpolatedUV = FVector2D::ZeroVector;
+		UGeometryScriptLibrary_MeshQueryFunctions::GetInterpolatedTriangleUV(DynamicMeshes[DataIndex], /*UVSetIndex=*/ UVChannel, TriangleId, BarycentricCoord, bHasValidUVs, InterpolatedUV);
+		if (bHasValidUVs)
+		{
+			check(OutPointData[DataIndex]);
+			OutPointData[DataIndex]->Metadata->InitializeOnSet(OutPoint.MetadataEntry);
+			UVAttributes[DataIndex]->SetValue(OutPoint.MetadataEntry, InterpolatedUV);
+		}
+	}
+
+	if (TriangleIdAttributes.IsValidIndex(DataIndex) && TriangleIdAttributes[DataIndex])
+	{
+		check(OutPointData[DataIndex]);
+		OutPointData[DataIndex]->Metadata->InitializeOnSet(OutPoint.MetadataEntry);
+		TriangleIdAttributes[DataIndex]->SetValue(OutPoint.MetadataEntry, TriangleId);
+	}
+}
+
+void FPCGMeshSamplerContext::SetPointColorAndDensity(SetPointDensityFunc SetPointDensityFuncPtr, int32 TriangleId, const FVector& BarycentricCoord, FPCGPoint& OutPoint, int32 DataIndex)
+{
+	FLinearColor Color;
+	bool bValidVertexColor = false;
+	UGeometryScriptLibrary_MeshQueryFunctions::GetInterpolatedTriangleVertexColor(DynamicMeshes[DataIndex], TriangleId, BarycentricCoord, FLinearColor::White, bValidVertexColor, Color);
+	if (bValidVertexColor)
+	{
+		OutPoint.Color = Color;
+		SetPointDensityFuncPtr(Color, OutPoint);
+	}
+	else
+	{
+		OutPoint.Density = 1.0f;
+	}
+}
+
 UPCGMeshSamplerSettings::UPCGMeshSamplerSettings()
 	: UPCGSettings()
 {
@@ -55,13 +285,19 @@ void UPCGMeshSamplerSettings::PostLoad()
 
 TArray<FPCGPinProperties> UPCGMeshSamplerSettings::InputPinProperties() const
 {
-	return TArray<FPCGPinProperties>{};
+	TArray<FPCGPinProperties> Properties;
+	if (bExtractMeshFromInput)
+	{
+		Properties.Emplace_GetRef(PCGPinConstants::DefaultInputLabel, EPCGDataType::Any).SetRequiredPin();
+	}
+
+	return Properties;
 }
 
 TArray<FPCGPinProperties> UPCGMeshSamplerSettings::OutputPinProperties() const
 {
 	TArray<FPCGPinProperties> Properties;
-	Properties.Emplace(PCGPinConstants::DefaultOutputLabel, EPCGDataType::Point, /*bInAllowMultipleConnections =*/ false, /*bAllowMultipleData =*/ false);
+	Properties.Emplace(PCGPinConstants::DefaultOutputLabel, EPCGDataType::Point, /*bInAllowMultipleConnections=*/bExtractMeshFromInput, /*bAllowMultipleData=*/bExtractMeshFromInput);
 	return Properties;
 }
 
@@ -79,20 +315,26 @@ FText UPCGMeshSamplerSettings::GetNodeTooltipText() const
 
 FPCGMeshSamplerContext::~FPCGMeshSamplerContext()
 {
-	// The context can be destroyed if the task is canceled. In that case, we need to notify the future that it should stop,
-	// and wait for it to finish (as the future uses some data stored in the context that will go dangling if the context is destroyed).
-	if (SamplingFuture.IsValid() && !SamplingFuture.IsReady())
+	// The context can be destroyed if the task is canceled. In that case, we need to notify the futures that they should stop,
+	// and wait for them to finish (as the futures use some data stored in the context that will go dangling if the context is destroyed).
+	StopSampling = true;
+	for (TFuture<bool>& SamplingFuture : SamplingFutures)
 	{
-		StopSampling = true;
-		SamplingFuture.Wait();
+		if (SamplingFuture.IsValid() && !SamplingFuture.IsReady())
+		{
+			SamplingFuture.Wait();
+		}
 	}
 }
 
 void FPCGMeshSamplerContext::AddExtraStructReferencedObjects(FReferenceCollector& Collector)
 {
-	if (DynamicMesh)
+	for (TObjectPtr<UDynamicMesh>& DynamicMesh : DynamicMeshes)
 	{
-		Collector.AddReferencedObject(DynamicMesh);
+		if (DynamicMesh)
+		{
+			Collector.AddReferencedObject(DynamicMesh);
+		}
 	}
 }
 
@@ -117,112 +359,178 @@ bool FPCGMeshSamplerElement::PrepareDataInternal(FPCGContext* InContext) const
 	check(Settings);
 
 	const TSoftObjectPtr<UStaticMesh> StaticMeshPtr = Settings->StaticMesh;
-	if (StaticMeshPtr.IsNull())
-	{
-		return true;
-	}
 
-	// 1. Request load for mesh. Return false if we need to wait, otherwise continue.
+	// 1. Request load for meshes/inputs. Return false if we need to wait, otherwise continue.
 	if (!Context->WasLoadRequested())
 	{
-		if (!Context->RequestResourceLoad(Context, { StaticMeshPtr.ToSoftObjectPath() }, !Settings->bSynchronousLoad))
+		const bool bNeedToWait = !Context->InitializeAndRequestLoad(PCGPinConstants::DefaultInputLabel, Settings->InputSource, { StaticMeshPtr.ToSoftObjectPath() }, /*bPersistAllData=*/false, /*bSilenceErrorOnEmptyObjectPath=*/true, Settings->bSynchronousLoad);
+
+		if (bNeedToWait)
 		{
 			return false;
 		}
 	}
 
-	UStaticMesh* StaticMesh = StaticMeshPtr.Get();
-
-	if (!StaticMesh)
+	for (const auto& [Path, DummyIndex, DummyIndex2] : Context->PathsToObjectsAndDataIndex)
 	{
-		PCGE_LOG(Error, GraphAndLog, FText::Format(LOCTEXT("MeshDoesNotExist", "Provided static mesh does not exist or could not be loaded: '{0}'"), FText::FromString(StaticMeshPtr.ToString())));
-		return true;
-	}
-
-	Context->DynamicMesh = NewObject<UDynamicMesh>();
+		UObject* Object = Path.ResolveObject();
+		if (!Object)
+		{
+			PCGE_LOG(Error, GraphAndLog, FText::Format(LOCTEXT("ObjectDoesNotExist", "Provided object does not exist or could not be loaded: '{0}'"), FText::FromString(Path.ToString())));
+			return true;
+		}
 
-	EGeometryScriptOutcomePins Outcome;
+		EGeometryScriptOutcomePins Outcome;
+		const FGeometryScriptMeshReadLOD MeshReadLOD{Settings->RequestedLODType, Settings->RequestedLODIndex };
+		USceneComponent* SceneComponent = Cast<USceneComponent>(Object);
 
-	UGeometryScriptLibrary_StaticMeshFunctions::CopyMeshFromStaticMesh(StaticMesh, Context->DynamicMesh, FGeometryScriptCopyMeshFromAssetOptions{},
-		FGeometryScriptMeshReadLOD{ Settings->RequestedLODType, Settings->RequestedLODIndex }, Outcome);
+		// @todo_pcg: Better support of multiple scene component on a given actor.
+		if (AActor* Actor = Cast<AActor>(Object))
+		{
+			SceneComponent = Actor->GetRootComponent();
+		}
 
-	if (Outcome == EGeometryScriptOutcomePins::Failure)
-	{
-		PCGE_LOG(Error, GraphAndLog, LOCTEXT("StaticToDynamicMeshFailed", "Static to Dynamic mesh failed"));
-		return true;
-	}
-	
-	if (Settings->bVoxelize)
-	{
-		FGeometryScriptSolidifyOptions SolidifyOptions{};
-		SolidifyOptions.GridParameters.GridCellSize = Settings->VoxelSize;
-		SolidifyOptions.GridParameters.SizeMethod = EGeometryScriptGridSizingMethod::GridCellSize;
+		if (SceneComponent)
+		{
+			Context->DynamicMeshes.Add(FPCGContext::NewObject_AnyThread<UDynamicMesh>(Context));
+			FTransform Transform;
+			FGeometryScriptCopyMeshFromComponentOptions Options{};
+			Options.bWantInstanceColors = true;
+			Options.RequestedLOD = MeshReadLOD;
 
-		UGeometryScriptLibrary_MeshVoxelFunctions::ApplyMeshSolidify(Context->DynamicMesh, SolidifyOptions);
+			UGeometryScriptLibrary_SceneUtilityFunctions::CopyMeshFromComponent(SceneComponent, Context->DynamicMeshes.Last(), Options, /*bTransformToWorld=*/false, Transform, Outcome);
+		}
+		else if (UStaticMesh* StaticMesh = Cast<UStaticMesh>(Object))
+		{
+			Context->DynamicMeshes.Add(FPCGContext::NewObject_AnyThread<UDynamicMesh>(Context));
+			FGeometryScriptCopyMeshFromAssetOptions Options{};
+			UGeometryScriptLibrary_StaticMeshFunctions::CopyMeshFromStaticMesh(StaticMesh, Context->DynamicMeshes.Last(), Options, MeshReadLOD, Outcome);
+		}
+		else
+		{
+			PCGE_LOG(Error, GraphAndLog, FText::Format(LOCTEXT("ObjectNotRightType", "Provided object '{0}' is not a supported type. Only supports StaticMesh/Actor/SceneComponent."), FText::FromString(Path.ToString())));
+			continue;
+		}
 
-		if (Settings->bRemoveHiddenTriangles)
+		if (Outcome == EGeometryScriptOutcomePins::Failure)
 		{
-			FGeometryScriptRemoveHiddenTrianglesOptions RemoveTriangleOptions{};
-			
-			UGeometryScriptLibrary_MeshRepairFunctions::RemoveHiddenTriangles(Context->DynamicMesh, RemoveTriangleOptions);
+			PCGLog::LogErrorOnGraph(FText::Format(LOCTEXT("ObjectToDynamicMeshFailed", "Object to Dynamic mesh failed for object {0}"), FText::FromString(Path.ToString())));
+			Context->DynamicMeshes.Last()->MarkAsGarbage();
+			Context->DynamicMeshes.RemoveAtSwap(Context->DynamicMeshes.Num() - 1);
 		}
 	}
 
+	// Reserve arrays. Add one more entry for the starting indices to have the total number of items to process.
 	switch (Settings->SamplingMethod)
 	{
 	case EPCGMeshSamplingMethod::OnePointPerVertex:
-	{
-		bool Dummy, Dummy2;
-		UGeometryScriptLibrary_MeshQueryFunctions::GetAllVertexPositions(Context->DynamicMesh, Context->Positions, /*bSkipGaps=*/false, /*bHasVertexIDGaps=*/Dummy);
-		UGeometryScriptLibrary_MeshVertexColorFunctions::GetMeshPerVertexColors(Context->DynamicMesh, Context->Colors, /*bIsValidColorSet=*/Dummy, /*bHasVertexIDGaps=*/Dummy2);
-		UGeometryScriptLibrary_MeshNormalsFunctions::GetMeshPerVertexNormals(Context->DynamicMesh, Context->Normals, /*bIsValidNormalSet=*/Dummy, /*bHasVertexIDGaps=*/Dummy2);
-		Context->Iterations = Context->Positions.List->Num();
+		Context->Positions.SetNum(Context->DynamicMeshes.Num());
+		Context->Colors.SetNum(Context->DynamicMeshes.Num());
+		Context->Normals.SetNum(Context->DynamicMeshes.Num());
+		Context->StartingIndices.SetNumZeroed(Context->DynamicMeshes.Num() + 1);
 		break;
-	}
 	case EPCGMeshSamplingMethod::OnePointPerTriangle:
-	{
-		bool Dummy;
-		UGeometryScriptLibrary_MeshQueryFunctions::GetAllTriangleIDs(Context->DynamicMesh, Context->TriangleIds, /*bHasTriangleIDGaps=*/Dummy);
-		Context->Iterations = Context->TriangleIds.List->Num();
+		Context->TriangleIds.SetNum(Context->DynamicMeshes.Num());
+		Context->StartingIndices.SetNumZeroed(Context->DynamicMeshes.Num() + 1);
 		break;
-	}
-	case EPCGMeshSamplingMethod::PoissonSampling:
-	{
-		// No preparation needed
+	default:
 		break;
 	}
-	default:
+	
+	for (int32 i = 0; i < Context->DynamicMeshes.Num(); ++i)
 	{
-		return true;
-	}
-	}
-
-	TArray<FPCGTaggedData>& Outputs = InContext->OutputData.TaggedData;
-	Context->OutPointData = NewObject<UPCGPointData>();
+		UDynamicMesh* DynamicMesh = Context->DynamicMeshes[i];
+		if (!DynamicMesh)
+		{
+			continue;
+		}
 
-	// It's not clear how to compute UVs for Vertices as they are part of multiple triangles. So disable for this mode. Same for triangle ids.
-	if (Settings->SamplingMethod != EPCGMeshSamplingMethod::OnePointPerVertex)
-	{
-		if (Settings->bExtractUVAsAttribute)
+		if (Settings->bVoxelize)
 		{
-			Context->UVAttribute = Context->OutPointData->Metadata->CreateAttribute<FVector2D>(Settings->UVAttributeName, FVector2D::ZeroVector, /*bAllowsInterpolation=*/true, /*bOverrideParent=*/true);
-			if (!Context->UVAttribute)
+			FGeometryScriptSolidifyOptions SolidifyOptions{};
+			SolidifyOptions.GridParameters.GridCellSize = Settings->VoxelSize;
+			SolidifyOptions.GridParameters.SizeMethod = EGeometryScriptGridSizingMethod::GridCellSize;
+
+			UGeometryScriptLibrary_MeshVoxelFunctions::ApplyMeshSolidify(DynamicMesh, SolidifyOptions);
+
+			if (Settings->bRemoveHiddenTriangles)
 			{
-				PCGE_LOG(Warning, GraphAndLog, FText::Format(LOCTEXT("AttributeUVFailed", "Failed to create attribute {0} for UVs. UVs won't be computed"), FText::FromName(Settings->UVAttributeName)));
+				FGeometryScriptRemoveHiddenTrianglesOptions RemoveTriangleOptions{};
+
+				UGeometryScriptLibrary_MeshRepairFunctions::RemoveHiddenTriangles(DynamicMesh, RemoveTriangleOptions);
 			}
 		}
 
-		if (Settings->bOutputTriangleIds)
+		TArray<FPCGTaggedData>& Outputs = InContext->OutputData.TaggedData;
+		UPCGPointData* CurrentOutPointData = FPCGContext::NewObject_AnyThread<UPCGPointData>(Context);
+		Context->OutPointData.Emplace(CurrentOutPointData);
+
+		int32 NumIterations = -1;
+
+		switch (Settings->SamplingMethod)
+		{
+		case EPCGMeshSamplingMethod::OnePointPerVertex:
+		{
+			bool Dummy, Dummy2;
+			UGeometryScriptLibrary_MeshQueryFunctions::GetAllVertexPositions(DynamicMesh, Context->Positions[i], /*bSkipGaps=*/false, /*bHasVertexIDGaps=*/Dummy);
+			UGeometryScriptLibrary_MeshVertexColorFunctions::GetMeshPerVertexColors(DynamicMesh, Context->Colors[i], /*bIsValidColorSet=*/Dummy, /*bHasVertexIDGaps=*/Dummy2);
+			UGeometryScriptLibrary_MeshNormalsFunctions::GetMeshPerVertexNormals(DynamicMesh, Context->Normals[i], /*bIsValidNormalSet=*/Dummy, /*bHasVertexIDGaps=*/Dummy2);
+			NumIterations = Context->Positions[i].List->Num();
+			break;
+		}
+		case EPCGMeshSamplingMethod::OnePointPerTriangle:
+		{
+			bool Dummy;
+			UGeometryScriptLibrary_MeshQueryFunctions::GetAllTriangleIDs(DynamicMesh, Context->TriangleIds[i], /*bHasTriangleIDGaps=*/Dummy);
+			NumIterations = Context->TriangleIds[i].List->Num();
+			break;
+		}
+		case EPCGMeshSamplingMethod::PoissonSampling:
+		{
+			// No preparation needed
+			break;
+		}
+		default:
+		{
+			PCGLog::LogErrorOnGraph(LOCTEXT("InvalidOperation", "Invalid sampling operation"), InContext);
+			return true;
+		}
+		}
+
+		if (NumIterations != -1)
 		{
-			Context->TriangleIdAttribute = Context->OutPointData->Metadata->CreateAttribute<int32>(Settings->TriangleIdAttributeName, 0, /*bAllowsInterpolation=*/false, /*bOverrideParent=*/true);
-			if (!Context->TriangleIdAttribute)
+			Context->StartingIndices[i + 1] = Context->StartingIndices[i] + NumIterations;
+			CurrentOutPointData->GetMutablePoints().SetNumUninitialized(NumIterations);
+		}
+
+		// It's not clear how to compute UVs for Vertices as they are part of multiple triangles. So disable for this mode. Same for triangle ids.
+		if (Settings->SamplingMethod != EPCGMeshSamplingMethod::OnePointPerVertex)
+		{
+			if (Settings->bExtractUVAsAttribute)
 			{
-				PCGE_LOG(Warning, GraphAndLog, FText::Format(LOCTEXT("AttributeTriangleIdFailed", "Failed to create attribute {0} for triangles ids. Triangle Ids won't be output"), FText::FromName(Settings->TriangleIdAttributeName)));
+				FPCGMetadataAttribute<FVector2D>* UVAttribute = CurrentOutPointData->Metadata->CreateAttribute<FVector2D>(Settings->UVAttributeName, FVector2D::ZeroVector, /*bAllowsInterpolation=*/true, /*bOverrideParent=*/true);
+				if (!UVAttribute)
+				{
+					PCGE_LOG(Warning, GraphAndLog, FText::Format(LOCTEXT("AttributeUVFailed", "Failed to create attribute {0} for UVs. UVs won't be computed"), FText::FromName(Settings->UVAttributeName)));
+				}
+
+				Context->UVAttributes.Emplace(UVAttribute);
+			}
+
+			if (Settings->bOutputTriangleIds)
+			{
+				FPCGMetadataAttribute<int32>* TriangleIdAttribute = CurrentOutPointData->Metadata->CreateAttribute<int32>(Settings->TriangleIdAttributeName, 0, /*bAllowsInterpolation=*/false, /*bOverrideParent=*/true);
+				if (!TriangleIdAttribute)
+				{
+					PCGE_LOG(Warning, GraphAndLog, FText::Format(LOCTEXT("AttributeTriangleIdFailed", "Failed to create attribute {0} for triangles ids. Triangle Ids won't be output"), FText::FromName(Settings->TriangleIdAttributeName)));
+				}
+
+				Context->TriangleIdAttributes.Emplace(TriangleIdAttribute);
 			}
 		}
-	}
 
-	Outputs.Emplace_GetRef().Data = Context->OutPointData;
+		Outputs.Emplace_GetRef().Data = CurrentOutPointData;
+	}
 
 	Context->bDataPrepared = true;
 	return true;
@@ -235,40 +543,15 @@ bool FPCGMeshSamplerElement::ExecuteInternal(FPCGContext* InContext) const
 
 	check(Context);
 
-	if (!Context->bDataPrepared)
+	if (!Context->bDataPrepared || Context->OutPointData.IsEmpty())
 	{
 		return true;
 	}
 
-	check(Context->OutPointData);
-
 	const UPCGMeshSamplerSettings* Settings = Context->GetInputSettings<UPCGMeshSamplerSettings>();
 	check(Settings);
 
 	bool bIsDone = false;
-	const bool bEnableTimeSlicing = true;
-	const int Seed = Context->GetSeed();
-
-	auto SetUVValueAndTriangleId = [Context, Settings](int32 TriangleId, const FVector& BarycentricCoord, FPCGPoint& OutPoint)
-	{
-		if (Context->UVAttribute)
-		{
-			bool bHasValidUVs = false;
-			FVector2D InterpolatedUV{};
-			UGeometryScriptLibrary_MeshQueryFunctions::GetInterpolatedTriangleUV(Context->DynamicMesh, /*UVSetIndex=*/ Settings->UVChannel, TriangleId, BarycentricCoord, bHasValidUVs, InterpolatedUV);
-			if (bHasValidUVs)
-			{
-				Context->OutPointData->Metadata->InitializeOnSet(OutPoint.MetadataEntry);
-				Context->UVAttribute->SetValue(OutPoint.MetadataEntry, InterpolatedUV);
-			}
-		}
-
-		if (Context->TriangleIdAttribute)
-		{
-			Context->OutPointData->Metadata->InitializeOnSet(OutPoint.MetadataEntry);
-			Context->TriangleIdAttribute->SetValue(OutPoint.MetadataEntry, TriangleId);
-		}
-	};
 
 	// Preparing the set function to extract the color to density.
 	auto SetPointDensityTo1 = [](const FLinearColor&, FPCGPoint& OutPoint){ OutPoint.Density = 1.0f; };
@@ -302,177 +585,27 @@ bool FPCGMeshSamplerElement::ExecuteInternal(FPCGContext* InContext) const
 		}
 	}
 
-	auto SetPointColorAndDensity = [Context, Settings, SetPointDensityPtr](int32 TriangleId, const FVector& BarycentricCoord, FPCGPoint& OutPoint)
-	{
-		FLinearColor Color;
-		bool bValidVertexColor = false;
-		UGeometryScriptLibrary_MeshQueryFunctions::GetInterpolatedTriangleVertexColor(Context->DynamicMesh, TriangleId, BarycentricCoord, FLinearColor::White, bValidVertexColor, Color);
-		if (bValidVertexColor)
-		{
-			OutPoint.Color = Color;
-			SetPointDensityPtr(Color, OutPoint);
-		}
-		else
-		{
-			OutPoint.Density = 1.0f;
-		}
-	};
-
 	switch (Settings->SamplingMethod)
 	{
 	case EPCGMeshSamplingMethod::OnePointPerVertex:
 	{
-		TRACE_CPUPROFILER_EVENT_SCOPE(FPCGMeshSamplerElement::Execute::OnePointPerVertex);
-
-		const TArray<FVector>& Positions = *Context->Positions.List.Get();
-		const TArray<FLinearColor>& Colors = *Context->Colors.List.Get();
-		const TArray<FVector>& Normals = *Context->Normals.List.Get();
-
-		auto IterationBody = [&Positions, &Colors, &Normals, Settings, Context, SetPointDensityPtr](int32 Index, FPCGPoint& OutPoint) -> bool
-		{
-			const FVector& Position = Positions[Index];
-			const FLinearColor& Color = Colors[Index];
-			const FVector& Normal = Normals[Index];
-
-			OutPoint = FPCGPoint{};
-			OutPoint.Transform = FTransform{ FRotationMatrix::MakeFromZ(Normal).Rotator(), Position, FVector{1.0, 1.0, 1.0} };
-			OutPoint.Color = Color;
-			OutPoint.Steepness = Settings->PointSteepness;
-
-			SetPointDensityPtr(Color, OutPoint);
-
-			UPCGBlueprintHelpers::SetSeedFromPosition(OutPoint);
-
-			return true;
-		};
-
-		bIsDone = FPCGAsync::AsyncProcessing<FPCGPoint>(&Context->AsyncState, Context->Iterations, Context->OutPointData->GetMutablePoints(), IterationBody, bEnableTimeSlicing);
+		bIsDone = PCGMeshSampler::SampleOnePointPerVertex(Settings, Context, SetPointDensityPtr);
 		break;
 	}
 	case EPCGMeshSamplingMethod::OnePointPerTriangle:
 	{
-		TRACE_CPUPROFILER_EVENT_SCOPE(FPCGMeshSamplerElement::Execute::OnePointPerTriangle);
-
-		const TArray<int32>& TriangleIds = *Context->TriangleIds.List.Get();
-
-		auto IterationBody = [DynamicMesh = Context->DynamicMesh, PointSteepness = Settings->PointSteepness, &TriangleIds, SetUVValueAndTriangleId, SetPointColorAndDensity](int32 Index, FPCGPoint& OutPoint) -> bool
-		{
-			const int32 TriangleId = TriangleIds[Index];
-
-			FVector Vertex1{}, Vertex2{}, Vertex3{};
-			bool bIsValidTriangle;
-
-			UGeometryScriptLibrary_MeshQueryFunctions::GetTrianglePositions(DynamicMesh, TriangleId, bIsValidTriangle, Vertex1, Vertex2, Vertex3);
-			const FVector Normal = UGeometryScriptLibrary_MeshQueryFunctions::GetTriangleFaceNormal(DynamicMesh, TriangleId, bIsValidTriangle);
-			const FVector Position = (Vertex1 + Vertex2 + Vertex3) / 3.0;
-
-			OutPoint = FPCGPoint{};
-			OutPoint.Transform = FTransform{ FRotationMatrix::MakeFromZ(Normal).Rotator(), Position, FVector{1.0, 1.0, 1.0} };
-			OutPoint.Steepness = PointSteepness;
-
-			FVector Dummy1, Dummy2, Dummy3;
-			FVector BarycentricCoord;
-			bool bIsValid = false;
-			UGeometryScriptLibrary_MeshQueryFunctions::ComputeTriangleBarycentricCoords(DynamicMesh, TriangleId, bIsValid, Position, Dummy1, Dummy2, Dummy3, BarycentricCoord);
-			
-			SetPointColorAndDensity(TriangleId, BarycentricCoord, OutPoint);
-			SetUVValueAndTriangleId(TriangleId, BarycentricCoord, OutPoint);
-
-			UPCGBlueprintHelpers::SetSeedFromPosition(OutPoint);
-
-			return true;
-		};
-
-		bIsDone = FPCGAsync::AsyncProcessing<FPCGPoint>(&Context->AsyncState, Context->Iterations, Context->OutPointData->GetMutablePoints(), IterationBody, bEnableTimeSlicing);
+		bIsDone = PCGMeshSampler::SampleOnePointPerTriangle(Settings, Context, SetPointDensityPtr);
 		break;
 	}
 	case EPCGMeshSamplingMethod::PoissonSampling:
 	{
-		// For Poisson sampling, we are calling a "all-in-one" function, where we don't have control for timeslicing.
-		// Since Poisson sampling can be expensive (depending on the radius used), we will do the sampling in a future, put this task to sleep,
-		// and wait for the sampling to wake us up.
-		if (!Context->SamplingFuture.IsValid())
-		{
-			// Put the task asleep
-			Context->bIsPaused = true;
-			Context->SamplingProgess = MakeUnique<FProgressCancel>();
-			Context->SamplingProgess->CancelF = [Context]() -> bool { return Context->StopSampling; };
-
-			auto SamplingFuture = [Settings, Context, Seed, SetUVValueAndTriangleId, SetPointColorAndDensity]() -> bool
-			{
-				TRACE_CPUPROFILER_EVENT_SCOPE(FPCGMeshSamplerElement::Execute::PoissonSampling);
-
-				UE::Geometry::FMeshSurfacePointSampling PointSampling;
-				PointSampling.SampleRadius = Settings->SamplingOptions.SamplingRadius;
-				PointSampling.MaxSamples = Settings->SamplingOptions.MaxNumSamples;
-				PointSampling.RandomSeed = PCGHelpers::ComputeSeed(Seed, Settings->SamplingOptions.RandomSeed);
-				PointSampling.SubSampleDensity = Settings->SamplingOptions.SubSampleDensity;
-
-				if (Settings->NonUniformSamplingOptions.MaxSamplingRadius > PointSampling.SampleRadius)
-				{
-					PointSampling.MaxSampleRadius = Settings->NonUniformSamplingOptions.MaxSamplingRadius;
-					PointSampling.SizeDistribution = UE::Geometry::FMeshSurfacePointSampling::ESizeDistribution(static_cast<int>(Settings->NonUniformSamplingOptions.SizeDistribution));
-					PointSampling.SizeDistributionPower = FMath::Clamp(Settings->NonUniformSamplingOptions.SizeDistributionPower, 1.0, 10.0);
-				}
-
-				PointSampling.bComputeBarycentrics = true;
-
-				PointSampling.ComputePoissonSampling(Context->DynamicMesh->GetMeshRef(), Context->SamplingProgess.Get());
-
-				if (Context->StopSampling)
-				{
-					return true;
-				}
-
-				TArray<FPCGPoint>& Points = Context->OutPointData->GetMutablePoints();
-				Points.Reserve(PointSampling.Samples.Num());
-
-				int Count = 0;
-				for (int32 i = 0; i < PointSampling.Samples.Num(); ++i)
-				{
-					UE::Geometry::FFrame3d& Sample = PointSampling.Samples[i];
-					// Avoid to check too many times
-					constexpr int CancelledCheckNum = 25;
-					if (++Count == CancelledCheckNum)
-					{
-						Count = 0;
-						if (Context->StopSampling)
-						{
-							return true;
-						}
-					}
-
-					const int32 TriangleId = PointSampling.TriangleIDs[i];
-					const FVector BarycentricCoords = PointSampling.BarycentricCoords[i];
-
-					FPCGPoint& OutPoint = Points.Emplace_GetRef();
-					OutPoint.Transform = Sample.ToTransform();
-					OutPoint.Steepness = Settings->PointSteepness;
-
-					SetPointColorAndDensity(TriangleId, BarycentricCoords, OutPoint);
-					SetUVValueAndTriangleId(TriangleId, BarycentricCoords, OutPoint);
-
-					UPCGBlueprintHelpers::SetSeedFromPosition(OutPoint);
-				}
-
-				// Unpause the task
-				Context->bIsPaused = false;
-				return true;
-			};
-
-			Context->SamplingFuture = Async(EAsyncExecution::ThreadPool, std::move(SamplingFuture));
-		}
-
-		bIsDone = Context->SamplingFuture.IsReady();
-		if (bIsDone)
-		{
-			Context->SamplingFuture.Reset();
-		}
-
+		bIsDone = PCGMeshSampler::PoissonSampling(Settings, Context, SetPointDensityPtr);
 		break;
 	}
 	default:
 	{
+		// Already logged in the PrepareData, we should never arrive there
+		checkNoEntry();
 		return true;
 	}
 	}
diff --git a/Engine/Plugins/PCGGeometryScriptInterop/Source/PCGGeometryScriptInterop/Public/Elements/PCGMeshSampler.h b/Engine/Plugins/PCGGeometryScriptInterop/Source/PCGGeometryScriptInterop/Public/Elements/PCGMeshSampler.h
--- a/Engine/Plugins/PCGGeometryScriptInterop/Source/PCGGeometryScriptInterop/Public/Elements/PCGMeshSampler.h
+++ b/Engine/Plugins/PCGGeometryScriptInterop/Source/PCGGeometryScriptInterop/Public/Elements/PCGMeshSampler.h
@@ -5,7 +5,7 @@
 #include "PCGContext.h"
 #include "PCGElement.h"
 #include "PCGSettings.h"
-#include "Async/PCGAsyncLoadingContext.h"
+#include "Elements/PCGLoadObjectsContext.h"
 
 #include "Async/Future.h"
 #include "GeometryScript/GeometryScriptTypes.h"
@@ -49,7 +49,7 @@ enum class EPCGColorChannel
 };
 
 /**
-* Sample points on a mesh
+* Sample points on a mesh.
 */
 UCLASS(BlueprintType, ClassGroup = (Procedural))
 class PCGGEOMETRYSCRIPTINTEROP_API UPCGMeshSamplerSettings : public UPCGSettings
@@ -79,13 +79,25 @@ protected:
 	// ~End UPCGSettings interface
 
 public:
+	/** 
+	* Can provide a list of inputs to sample the meshes from. It can be a list of StaticMeshes, a list of Actors that have a Scene Component (like a Static Mesh Component), or a list of Scene Components directly.
+	* Geometry Script needs to be able to extract a dynamic mesh from this scene component (so won't work for ISMCs for example) and for now will work only with a single scene component.
+	* Each entry (either in the same data or seperate data) will produce a unique output data.
+	*/
 	UPROPERTY(BlueprintReadWrite, EditAnywhere, Category = "Settings")
-	EPCGMeshSamplingMethod SamplingMethod = EPCGMeshSamplingMethod::OnePointPerTriangle;
+	bool bExtractMeshFromInput = false;
 
 	/** Soft Object Path to the mesh to sample from. Will be loaded. */
-	UPROPERTY(BlueprintReadWrite, EditAnywhere, Category = "Settings", meta = (PCG_Overridable))
+	UPROPERTY(BlueprintReadWrite, EditAnywhere, Category = "Settings", meta = (PCG_Overridable, EditCondition = "!bExtractMeshFromInput", EditConditionHides, ShowAfter = "bExtractMeshFromInput"))
 	TSoftObjectPtr<UStaticMesh> StaticMesh;
 
+	/** Selector to read data from. */
+	UPROPERTY(BlueprintReadWrite, EditAnywhere, Category = "Settings", meta = (PCG_Overridable, EditCondition = "bExtractMeshFromInput", EditConditionHides, ShowAfter = "bExtractMeshFromInput"))
+	FPCGAttributePropertyInputSelector InputSource;
+
+	UPROPERTY(BlueprintReadWrite, EditAnywhere, Category = "Settings")
+	EPCGMeshSamplingMethod SamplingMethod = EPCGMeshSamplingMethod::OnePointPerTriangle;
+
 	/** Will extract the color channel into the density. */
 	UPROPERTY(BlueprintReadWrite, EditAnywhere, Category = "Settings|Color & Density", meta = (PCG_Overridable, InlineEditConditionToggle, PCG_OverrideAliases="bUseRedAsDensity"))
 	bool bUseColorChannelAsDensity = false;
@@ -159,39 +171,44 @@ protected:
 /**
 * Extra context to store all the data that need to be kept between multiple executions (time slicing)
 */
-struct FPCGMeshSamplerContext : public FPCGContext, public IPCGAsyncLoadingContext
+struct FPCGMeshSamplerContext : public FPCGLoadObjectsFromPathContext
 {
 public:
 	~FPCGMeshSamplerContext();
 
+	using SetPointDensityFunc = void(*)(const FLinearColor&, FPCGPoint&);
+
+	void SetUVValueAndTriangleId(int32 UVChannel, int32 TriangleId, const FVector& BarycentricCoord, FPCGPoint& OutPoint, int32 DataIndex);
+	void SetPointColorAndDensity(SetPointDensityFunc SetPointDensityFuncPtr, int32 TriangleId, const FVector& BarycentricCoord, FPCGPoint& OutPoint, int32 DataIndex);
+
 protected:
 	virtual void AddExtraStructReferencedObjects(FReferenceCollector& Collector) override;
 	
 public:
-	// Dynamic mesh. Will be added to root.
-	TObjectPtr<UDynamicMesh> DynamicMesh;
+	// Dynamic meshes. Will be added to root.
+	TArray<TObjectPtr<UDynamicMesh>> DynamicMeshes;
 
-	// Lists extracted from the mesh
-	FGeometryScriptVectorList Positions;
-	FGeometryScriptColorList Colors;
-	FGeometryScriptVectorList Normals;
-	FGeometryScriptIndexList TriangleIds;
+	// Lists extracted from the meshes
+	TArray<FGeometryScriptVectorList> Positions;
+	TArray<FGeometryScriptColorList> Colors;
+	TArray<FGeometryScriptVectorList> Normals;
+	TArray<FGeometryScriptIndexList> TriangleIds;
 
-	// Output point data
-	UPCGPointData* OutPointData = nullptr;
+	// Output point data.
+	TArray<UPCGPointData*> OutPointData;
 
 	// Optional attributes.
-	FPCGMetadataAttribute<FVector2D>* UVAttribute = nullptr;
-	FPCGMetadataAttribute<int32>* TriangleIdAttribute = nullptr;
+	TArray<FPCGMetadataAttribute<FVector2D>*> UVAttributes;
+	TArray<FPCGMetadataAttribute<int32>*> TriangleIdAttributes;
 
-	// For Poisson sampling, we are starting a future that is not framebound
-	// Store the future and synchronisation items in the context
-	TFuture<bool> SamplingFuture;
+	// For Poisson sampling, we are starting futures that are not framebound
+	// Store the futures and synchronisation items in the context
+	TArray<TFuture<bool>> SamplingFutures;
 	std::atomic<bool> StopSampling = false;
 	TUniquePtr<FProgressCancel> SamplingProgess;
 
-	// Number of iterations to be done
-	int32 Iterations = 0;
+	// Starting indices for each different object to sample. If we have 3 meshes of 20 elements (like vertices) each, the array will be [0, 20, 40, 60]
+	TArray<int32> StartingIndices;
 
 	// Set to true if prepared data succeeded.
 	bool bDataPrepared = false;
