[PCG] Spline Cross Section Generator Node

diff --git a/Engine/Plugins/PCGGeometryScriptInterop/Source/PCGGeometryScriptInterop/Private/Elements/PCGSplineCrossSectionGenerator.cpp b/Engine/Plugins/PCGGeometryScriptInterop/Source/PCGGeometryScriptInterop/Private/Elements/PCGSplineCrossSectionGenerator.cpp
new file mode 100644
--- /dev/null
+++ b/Engine/Plugins/PCGGeometryScriptInterop/Source/PCGGeometryScriptInterop/Private/Elements/PCGSplineCrossSectionGenerator.cpp
@@ -0,0 +1,478 @@
+// Copyright Epic Games, Inc. All Rights Reserved.
+
+#include "Elements/PCGSplineCrossSectionGenerator.h"
+
+#include "PCGContext.h"
+#include "PCGComponent.h"
+#include "Data/PCGPrimitiveData.h"
+#include "Data/PCGSplineData.h"
+#include "Helpers/PCGHelpers.h"
+
+#include "UDynamicMesh.h"
+#include "Algo/Find.h"
+#include "Algo/ForEach.h"
+#include "Components/SplineComponent.h"
+#include "GeometryScript/MeshBooleanFunctions.h"
+#include "GeometryScript/MeshModelingFunctions.h"
+#include "GeometryScript/MeshQueryFunctions.h"
+#include "GeometryScript/MeshSelectionFunctions.h"
+#include "GeometryScript/MeshSelectionQueryFunctions.h"
+#include "GeometryScript/MeshSimplifyFunctions.h"
+#include "GeometryScript/SceneUtilityFunctions.h"
+
+#define LOCTEXT_NAMESPACE "PCGSplineCrossSectionGeneratorElement"
+
+namespace SplineCrossSectionGeneratorConstants
+{
+	const FName DefaultExtrusionVectorAttributeName(TEXT("ExtrusionVector"));
+	constexpr FGeometryScriptMeshPlaneCutOptions CutPlaneOptions = FGeometryScriptMeshPlaneCutOptions(/*bFillHoles=*/false, /*bFillSpans=*/false, /*bFlipCutSide=*/true, 1);
+	constexpr FGeometryScriptCopyMeshFromComponentOptions CopyMeshFromComponentOptions(/*bWantNormals=*/false, /*bWantTangents=*/false);
+}
+
+struct FCrossSection
+{
+	int Tier;
+	double Height;
+	TArray<FVector> PointLocations;
+};
+
+UPCGSplineCrossSectionGeneratorSettings::UPCGSplineCrossSectionGeneratorSettings()
+{
+	ExtrusionVectorAttribute.SetAttributeName(SplineCrossSectionGeneratorConstants::DefaultExtrusionVectorAttributeName);
+}
+
+TArray<FPCGPinProperties> UPCGSplineCrossSectionGeneratorSettings::InputPinProperties() const
+{
+	TArray<FPCGPinProperties> PinProperties;
+	PinProperties.Emplace(PCGPinConstants::DefaultInputLabel, EPCGDataType::Primitive);
+	return PinProperties;
+}
+
+TArray<FPCGPinProperties> UPCGSplineCrossSectionGeneratorSettings::OutputPinProperties() const
+{
+	TArray<FPCGPinProperties> PinProperties;
+	PinProperties.Emplace(PCGPinConstants::DefaultOutputLabel, EPCGDataType::Spline);
+	return PinProperties;
+}
+
+FPCGElementPtr UPCGSplineCrossSectionGeneratorSettings::CreateElement() const
+{
+	return MakeShared<FPCGSplineCrossSectionGeneratorElement>();
+}
+
+bool FPCGSplineCrossSectionGeneratorElement::ExecuteInternal(FPCGContext* Context) const
+{
+	TRACE_CPUPROFILER_EVENT_SCOPE(FPCGSplineCrossSectionGeneratorElement::Execute);
+
+	const UPCGSplineCrossSectionGeneratorSettings* Settings = Context->GetInputSettings<UPCGSplineCrossSectionGeneratorSettings>();
+	const TArray<FPCGTaggedData> PrimitiveInputs = Context->InputData.GetInputsByPin(PCGPinConstants::DefaultInputLabel);
+
+	if (PrimitiveInputs.IsEmpty())
+	{
+		return true;
+	}
+
+	// TODO: Could also take in PCGVolume
+	TArray<const UPCGPrimitiveData*> PrimitiveDataCollection;
+	PrimitiveDataCollection.Reserve(PrimitiveInputs.Num());
+
+	for (const FPCGTaggedData& TaggedData : PrimitiveInputs)
+	{
+		if (const UPCGPrimitiveData* PrimitiveData = Cast<UPCGPrimitiveData>(TaggedData.Data))
+		{
+			PrimitiveDataCollection.Add(PrimitiveData);
+		}
+	}
+
+	if (PrimitiveDataCollection.IsEmpty())
+	{
+		return true;
+	}
+
+	UDynamicMesh* DynamicMesh = FPCGContext::NewObject_AnyThread<UDynamicMesh>(Context);
+#if WITH_EDITOR
+	UGeometryScriptDebug* DynamicMeshDebug = FPCGContext::NewObject_AnyThread<UGeometryScriptDebug>(Context);
+#else
+	UGeometryScriptDebug* DynamicMeshDebug = nullptr;
+#endif // WITH_EDITOR
+
+	// Collect all the primitives and append them to the dynamic mesh
+	for (const UPCGPrimitiveData* PrimitiveData : PrimitiveDataCollection)
+	{
+		TWeakObjectPtr<UPrimitiveComponent> PrimitivePtr = PrimitiveData->GetComponent();
+		if (!PrimitivePtr.IsValid())
+		{
+			continue;
+		}
+
+		// Convert from scene component to mesh and begin boolean operation
+		if (USceneComponent* SceneComponent = Cast<USceneComponent>(PrimitivePtr.Pin().Get()))
+		{
+			bool bFirstMesh = DynamicMesh->IsEmpty();
+
+			UDynamicMesh* ComponentMesh = NewObject<UDynamicMesh>();
+			EGeometryScriptOutcomePins Outcome;
+			FTransform PrimitiveTransform;
+			// TODO: If we have repeats, ie. components that match content/transform, we could skip them.
+			UGeometryScriptLibrary_SceneUtilityFunctions::CopyMeshFromComponent(
+				SceneComponent,
+				(bFirstMesh ? DynamicMesh : ComponentMesh),
+				SplineCrossSectionGeneratorConstants::CopyMeshFromComponentOptions,
+				/*bTransformToWorld=*/true,
+				PrimitiveTransform,
+				Outcome,
+				DynamicMeshDebug);
+
+			// TODO: Investigate if merging first and boolean the final result would be faster.
+			if (!bFirstMesh)
+			{
+				UGeometryScriptLibrary_MeshBooleanFunctions::ApplyMeshBoolean(
+					DynamicMesh,
+					FTransform::Identity,
+					ComponentMesh,
+					FTransform::Identity,
+					EGeometryScriptBooleanOperation::Union,
+					FGeometryScriptMeshBooleanOptions(), // TODO: check defaults
+					DynamicMeshDebug);
+			}
+		}
+	}
+
+	// Reduce vertex count by simplifying coplanar triangles. Also removes index gaps.
+	UGeometryScriptLibrary_MeshSimplifyFunctions::ApplySimplifyToPlanar(DynamicMesh, FGeometryScriptPlanarSimplifyOptions(), DynamicMeshDebug);
+
+	FGeometryScriptVectorList VertexList;
+	bool bHasGapsDummy = false;
+	UGeometryScriptLibrary_MeshQueryFunctions::GetAllVertexPositions(DynamicMesh, VertexList, /*bSkipGaps=*/false, bHasGapsDummy);
+
+	if (!VertexList.List.IsValid() || VertexList.List->IsEmpty())
+	{
+		return true;
+	}
+
+	// No cross-sections will be generated. Early out.
+	if (!Settings->bEnableFindFeatures && !Settings->bEnableTierSlicing)
+	{
+		return true;
+	}
+
+	// A safeguard to prevent analysis on a significantly large mesh.
+	if (Settings->bEnableFindFeatures && VertexList.List->Num() > Settings->MaxMeshVertexCount)
+	{
+		PCGLog::LogWarningOnGraph(LOCTEXT("MaxMeshVertexCount", "Combined mesh exceeds Max Mesh Vertex Count and will be skipped."), Context);
+		return true;
+	}
+
+	FVector SliceOrigin = FVector::ZeroVector;
+	const FVector SliceDirection = Settings->SliceDirection.GetSafeNormal();
+
+	double MinProjectionScalar = std::numeric_limits<double>::max();
+	double MaxProjectionScalar = std::numeric_limits<double>::min();
+
+	// Project vertices along the slice direction and find the minimum height
+	TArray<double> ProjectionScalars;
+	ProjectionScalars.Reserve(VertexList.List->Num());
+	for (const FVector& Vertex : *VertexList.List)
+	{
+		// Implementation note: since the direction vector is guaranteed to be normalized, we can use the projection scalar as distance
+		double ProjectionScalar = FVector::DotProduct(Vertex, SliceDirection);
+		ProjectionScalars.Emplace(ProjectionScalar);
+
+		// Find the min and max projected "height" to define the slicing extents
+		if (MinProjectionScalar > ProjectionScalar)
+		{
+			MinProjectionScalar = ProjectionScalar;
+			// Find the vertex for which we'll start slicing
+			SliceOrigin = Vertex;
+		}
+
+		MaxProjectionScalar = FMath::Max(MaxProjectionScalar, ProjectionScalar);
+	}
+
+	// Convert from projection scalar to height from origin
+	TMap<double, int> HeightMap;
+	for (double ProjectionScalar : ProjectionScalars)
+	{
+		// TODO: It would be better to compare an epsilon, in the case of sub 1cm meshes.
+		double Height = FMath::Floor(ProjectionScalar - MinProjectionScalar);
+		++HeightMap.FindOrAdd(Height, 0);
+	}
+
+	TArray<double> HeightKeys;
+	HeightMap.GetKeys(HeightKeys);
+
+	TArray<double, TInlineAllocator<16>> TierHeights;
+	// Find heights with coplanar vertices using the height counts
+	if (Settings->bEnableFindFeatures)
+	{
+		Algo::CopyIf(HeightKeys, TierHeights, [&HeightMap, MinVertices = Settings->MinimumCoplanarVertices](double Key)
+		{
+			return HeightMap[Key] >= MinVertices;
+		});
+	}
+
+	// Slice from the base of the mesh to the highest point by slice direction at the given interval
+	if (Settings->bEnableTierSlicing)
+	{
+		double TierSlicingResolution = FMath::Max(Settings->TierSlicingResolution, 100.0);
+		for (double Height = 0; Height <= (MaxProjectionScalar - MinProjectionScalar); Height += TierSlicingResolution)
+		{
+			TierHeights.AddUnique(Height);
+		}
+	}
+
+	if (TierHeights.IsEmpty())
+	{
+		return true;
+	}
+
+	TierHeights.Sort([](double LHS, double RHS) { return LHS < RHS; });
+
+	if (Settings->bEnableTierMerging)
+	{
+		// Filter the heights by threshold
+		TArray<double, TInlineAllocator<16>> FilteredTiers;
+		FilteredTiers.Reserve(TierHeights.Num());
+		FilteredTiers.Add(TierHeights[0]);
+
+		// TODO: This currently always merges "down", but it could also alternatively blend with an average, merge "up", etc
+		for (int TierIndex = 1; TierIndex < TierHeights.Num(); ++TierIndex)
+		{
+			if (TierHeights[TierIndex] - FilteredTiers.Last() > Settings->TierMergingThreshold)
+			{
+				FilteredTiers.Add(TierHeights[TierIndex]);
+			}
+		}
+
+		TierHeights = MoveTemp(FilteredTiers);
+	}
+
+	// Set up the slice plane for GeometryScript, represented by an FTransform
+	FTransform SlicePlaneTransform;
+	SlicePlaneTransform.SetLocation(SliceOrigin);
+	FQuat AdjustedRotation = FQuat::FindBetween(FVector::UpVector, SliceDirection);
+	SlicePlaneTransform.SetRotation(AdjustedRotation);
+
+	TArray<FCrossSection> CrossSections;
+	for (int TierIndex = 0; TierIndex < TierHeights.Num() - 1; ++TierIndex)
+	{
+		FVector SliceLocation = SliceOrigin + SliceDirection * TierHeights[TierIndex];
+		SlicePlaneTransform.SetLocation(SliceLocation);
+
+		// Cuts the mesh at the specified plane, leaving a hole
+		UGeometryScriptLibrary_MeshBooleanFunctions::ApplyMeshPlaneCut(DynamicMesh, SlicePlaneTransform, SplineCrossSectionGeneratorConstants::CutPlaneOptions, DynamicMeshDebug);
+
+		// Break now as there is no point in continuing along this invalid mesh
+		if (!ensure(DynamicMesh->GetTriangleCount() >= 1))
+		{
+			PCGLog::LogWarningOnGraph(FText::Format(LOCTEXT("DynamicMeshInvalid", "Dynamic Mesh Invalid at tier: {0}"), FText::AsNumber(TierIndex)), Context);
+			break;
+		}
+
+		// The cut algorithm does not result in simplified planes, so simplify it
+		UGeometryScriptLibrary_MeshSimplifyFunctions::ApplySimplifyToPlanar(DynamicMesh, FGeometryScriptPlanarSimplifyOptions(), DynamicMeshDebug); // TODO: verify defaults
+
+		// TODO: Investigate if there's a simpler API call to get the cut plane triangles (vertices)
+		// Since select box is AABB, in order to select a slice, select with a plane just before and subtract the one just after
+		FGeometryScriptMeshSelection CurrentSelection{};
+		UGeometryScriptLibrary_MeshSelectionFunctions::SelectMeshElementsWithPlane(
+			DynamicMesh,
+			CurrentSelection,
+			SliceLocation - (SliceDirection * UE_DOUBLE_KINDA_SMALL_NUMBER), // Move back a bit
+			SliceDirection,
+			EGeometryScriptMeshSelectionType::Vertices);
+
+		FGeometryScriptMeshSelection ExclusiveSelection{};
+		UGeometryScriptLibrary_MeshSelectionFunctions::SelectMeshElementsWithPlane(
+			DynamicMesh,
+			ExclusiveSelection,
+			SliceLocation + (SliceDirection * UE_DOUBLE_KINDA_SMALL_NUMBER), // Move forward a bit
+			SliceDirection,
+			EGeometryScriptMeshSelectionType::Vertices);
+
+		UGeometryScriptLibrary_MeshSelectionFunctions::CombineMeshSelections(
+			CurrentSelection,
+			ExclusiveSelection,
+			CurrentSelection,
+			EGeometryScriptCombineSelectionMode::Subtract);
+
+		TArray<FGeometryScriptIndexList> VertexIndexLists;
+		TArray<FGeometryScriptPolyPath> PolyPaths;
+		int NumLoops;
+		bool bFoundErrors;
+		// TODO: Find a more direct way to determine multiple poly paths--one for each island--and order the vertices
+		// Implementation note: Using Boundary Loops to create multiple 2D poly paths and order points.
+		UGeometryScriptLibrary_MeshSelectionQueryFunctions::GetMeshSelectionBoundaryLoops(
+			DynamicMesh,
+			CurrentSelection,
+			VertexIndexLists,
+			PolyPaths,
+			NumLoops,
+			bFoundErrors,
+			DynamicMeshDebug);
+
+#if WITH_EDITOR
+		// At this point, we're done with GeometryScript for this iteration. Print errors if they occur.
+		for (const FGeometryScriptDebugMessage& Message : DynamicMeshDebug->Messages)
+		{
+			PCGLog::LogWarningOnGraph(FText::Format(LOCTEXT("GeometryScriptError", "GeometryScript Error: {0}"), Message.Message), Context);
+		}
+#endif // WITH_EDITOR
+
+		if (NumLoops < 1 || PolyPaths.IsEmpty())
+		{
+			continue;
+		}
+
+		// Shoelace formula - https://en.wikipedia.org/wiki/Shoelace_formula.
+		// Copy PointLocations in by value to be used as a temporary container.
+		auto ComputePolyPathArea = [&AdjustedRotation](TArray<FVector> PointLocations)
+		{
+			const int NumPoints = PointLocations.Num();
+			if (NumPoints < 3)
+			{
+				return 0.0;
+			}
+
+			// Transform the points to 2D referential in order to calculate the surface area
+			Algo::ForEach(PointLocations, [&AdjustedRotation](FVector& Location) { Location = AdjustedRotation.Inverse().RotateVector(Location); });
+
+			double Area = 0.0;
+			for (int PointIndex = 0; PointIndex < PointLocations.Num() - 1; ++PointIndex)
+			{
+				Area += PointLocations[PointIndex].X * PointLocations[PointIndex + 1].Y - PointLocations[PointIndex + 1].X * PointLocations[PointIndex].Y;
+			}
+
+			// Calculate the last term (last and first point) separately
+			Area += PointLocations[NumPoints - 1].X * PointLocations[0].Y - PointLocations[0].X * PointLocations[NumPoints - 1].Y;
+
+			// The sign of the area could be used to determine winding, if needed in the future
+			return FMath::Abs(Area) * 0.5;
+		};
+
+		for (int LoopIndex = 0; LoopIndex < NumLoops; ++LoopIndex)
+		{
+			if (!ensure(PolyPaths[LoopIndex].bClosedLoop))
+			{
+				continue;
+			}
+
+			// Do a planar check of one of the vertices and eliminate poly paths outside our cut plane
+			// TODO: This is a crutch to keep the boundary loops from creeping.
+			double DotProduct = FMath::Abs(FVector::DotProduct(PolyPaths[LoopIndex].Path->Last() - SliceLocation, SliceDirection));
+			if (DotProduct > UE_DOUBLE_KINDA_SMALL_NUMBER)
+			{
+				continue;
+			}
+
+			if (Settings->bEnableMinAreaCulling && ComputePolyPathArea(*PolyPaths[LoopIndex].Path) <= Settings->MinAreaCullingThreshold * 100.0)
+			{
+				continue;
+			}
+
+			FCrossSection& CurrentTier = CrossSections.Emplace_GetRef();
+			CurrentTier.Tier = TierIndex;
+			CurrentTier.Height = TierHeights[TierIndex + 1] - TierHeights[TierIndex];
+			CurrentTier.PointLocations.Append(*PolyPaths[LoopIndex].Path);
+		}
+	}
+
+	// Filter unintended excess cross-sections
+	if (Settings->bRemoveRedundantSections)
+	{
+		TArray<int, TInlineAllocator<16>> RedundantSections;
+		// Implementation note: the cross-sections are sorted by height here, so the highest can be taken
+		CrossSections.Sort([](const FCrossSection& LHS, const FCrossSection& RHS) { return LHS.Tier < RHS.Tier; });
+
+		for (int FirstIndex = 0; FirstIndex < CrossSections.Num(); ++FirstIndex)
+		{
+			if (RedundantSections.Contains(FirstIndex))
+			{
+				continue;
+			}
+
+			for (int SecondIndex = CrossSections.Num() - 1; SecondIndex > FirstIndex; --SecondIndex)
+			{
+				if (RedundantSections.Contains(SecondIndex))
+				{
+					continue;
+				}
+
+				FCrossSection& FirstSection = CrossSections[FirstIndex];
+				const FCrossSection& SecondSection = CrossSections[SecondIndex];
+
+				if (SecondSection.Tier < FirstSection.Tier)
+				{
+					break;
+				}
+
+				// Disqualify based on point count
+				if (FirstSection.PointLocations.Num() != SecondSection.PointLocations.Num())
+				{
+					continue;
+				}
+
+				// Check if the cross-section's vertices are all equal to find if it's redundant for culling.
+				// Implementation Note: This behavior is not always intuitive with any mesh and will cull even if there are tiers in between.
+				// TODO: Some future options might include checking for intermediate tiers between, etc. However, we need to be mindful of merge order.
+				bool bCrossSectionIsEqual = true;
+				for (int PointIndex = 0; PointIndex < FirstSection.PointLocations.Num(); ++PointIndex)
+				{
+					if (!Algo::FindByPredicate(SecondSection.PointLocations, [FirstLocation = FirstSection.PointLocations[PointIndex]](const FVector& SecondLocation)
+					{
+						return FMath::IsNearlyEqual(FirstLocation.X, SecondLocation.X) && FMath::IsNearlyEqual(FirstLocation.Y, SecondLocation.Y);
+					}))
+					{
+						bCrossSectionIsEqual = false;
+						break;
+					}
+				}
+
+				if (bCrossSectionIsEqual)
+				{
+					RedundantSections.Add(SecondIndex);
+					// Combine the heights of both sections
+					FirstSection.Height += SecondSection.Height;
+				}
+			}
+		}
+
+		// Sort inverse, because they will be remove swapped from the actual array
+		RedundantSections.Sort([](int LHS, int RHS) { return LHS > RHS; });
+
+		for (int Index : RedundantSections)
+		{
+			CrossSections.RemoveAtSwap(Index);
+		}
+	}
+
+	// Create the splines from the tiers
+	for (const FCrossSection& CrossSection : CrossSections)
+	{
+		TArray<FSplinePoint> SplinePoints;
+		SplinePoints.Reserve(CrossSection.PointLocations.Num());
+		FPCGTaggedData& OutputData = Context->OutputData.TaggedData.Emplace_GetRef();
+
+		for (int Index = 0; Index < CrossSection.PointLocations.Num(); ++Index)
+		{
+			SplinePoints.Emplace(static_cast<float>(Index),
+				CrossSection.PointLocations[Index],
+				FVector::ZeroVector,
+				FVector::ZeroVector,
+				FRotator::ZeroRotator,
+				FVector::OneVector,
+				ESplinePointType::Linear);
+		}
+
+		UPCGSplineData* OutSplineData = FPCGContext::NewObject_AnyThread<UPCGSplineData>(Context);
+		OutSplineData->Initialize(SplinePoints, true, FTransform::Identity);
+		OutputData.Data = OutSplineData;
+
+		check(OutSplineData->Metadata);
+		OutSplineData->Metadata->CreateAttribute(Settings->ExtrusionVectorAttribute.GetName(), FVector(0, 0, CrossSection.Height), /*bAllowsInterpolation=*/false, /*bOverridesParent=*/true);
+	}
+
+	return true;
+}
+
+#undef LOCTEXT_NAMESPACE
diff --git a/Engine/Plugins/PCGGeometryScriptInterop/Source/PCGGeometryScriptInterop/Public/Elements/PCGSplineCrossSectionGenerator.h b/Engine/Plugins/PCGGeometryScriptInterop/Source/PCGGeometryScriptInterop/Public/Elements/PCGSplineCrossSectionGenerator.h
new file mode 100644
--- /dev/null
+++ b/Engine/Plugins/PCGGeometryScriptInterop/Source/PCGGeometryScriptInterop/Public/Elements/PCGSplineCrossSectionGenerator.h
@@ -0,0 +1,89 @@
+// Copyright Epic Games, Inc. All Rights Reserved.
+
+#pragma once
+
+#include "PCGSettings.h"
+
+#include "PCGSplineCrossSectionGenerator.generated.h"
+
+/**
+ * Creates a spline cross-section of one more primitives based on vertex features, slice resolution, or both.
+ */
+UCLASS(BlueprintType, ClassGroup = (Procedural))
+class UPCGSplineCrossSectionGeneratorSettings : public UPCGSettings
+{
+	GENERATED_BODY()
+
+public:
+	UPCGSplineCrossSectionGeneratorSettings();
+
+	//~Begin UPCGSettings interface
+#if WITH_EDITOR
+	virtual FName GetDefaultNodeName() const override { return FName(TEXT("SplineCrossSectionGenerator")); }
+	virtual FText GetDefaultNodeTitle() const override { return NSLOCTEXT("PCGSplineCrossSectionGeneratorElement", "NodeTitle", "Spline Cross Section Generator"); }
+	virtual FText GetNodeTooltipText() const override { return NSLOCTEXT("PCGSplineCrossSectionGeneratorElement", "NodeTooltip", "Creates a spline cross-section of one more primitives based on vertex features, slice resolution, or both."); }
+	virtual EPCGSettingsType GetType() const override { return EPCGSettingsType::Spatial; }
+#endif
+
+protected:
+	virtual TArray<FPCGPinProperties> InputPinProperties() const override;
+	virtual TArray<FPCGPinProperties> OutputPinProperties() const override;
+	virtual FPCGElementPtr CreateElement() const override;
+	//~End UPCGSettings interface
+
+public:
+	/** Slicing will happen from the minimum vertex along this direction vector (normalized). Useful primarily with the Tier Slicing Resolution. */
+	UPROPERTY(BlueprintReadWrite, EditAnywhere, Category = Settings, meta = (PCG_Overridable))
+	FVector SliceDirection = FVector::UpVector;
+
+	/** The attribute that will be populated with each cross-section's extrusion vector. */
+	UPROPERTY(BlueprintReadWrite, EditAnywhere, Category = Settings, meta = (PCG_Overridable))
+	FPCGAttributePropertyOutputSelector ExtrusionVectorAttribute;
+
+	/** Find primitive tier "features", which consist of a number of co-planar vertices. */
+	UPROPERTY(BlueprintReadWrite, EditAnywhere, Category = Settings, meta = (PCG_Overridable))
+	bool bEnableFindFeatures = true;
+
+	/** The minimum required number of vertices that must be co-planar in order to be considered a tier "feature". */
+	UPROPERTY(BlueprintReadWrite, EditAnywhere, Category = Settings, meta = (PCG_Overridable, EditCondition = "bEnableFindFeatures", EditConditionHides))
+	int32 MinimumCoplanarVertices = 3;
+
+	/** Create tiers at a specified resolution. */
+	UPROPERTY(BlueprintReadWrite, EditAnywhere, Category = Settings, meta = (PCG_Overridable, InlineEditConditionToggle))
+	bool bEnableTierSlicing = false;
+
+	/** A tier will be created at this interval (in cm). */
+	UPROPERTY(BlueprintReadWrite, EditAnywhere, Category = Settings, meta = (PCG_Overridable, EditCondition = "bEnableTierSlicing", ClampMin = "100"))
+	double TierSlicingResolution = 500.0;
+
+	/** Cull tiers that are within a specified threshold. */
+	UPROPERTY(BlueprintReadWrite, EditAnywhere, Category = Settings, meta = (PCG_Overridable, InlineEditConditionToggle))
+	bool bEnableTierMerging = false;
+
+	/** If a tier is within this distance (in cm) of the previous tier, it will be culled. */
+	UPROPERTY(BlueprintReadWrite, EditAnywhere, Category = Settings, meta = (PCG_Overridable, EditCondition = "bEnableTierMerging"))
+	double TierMergingThreshold = 100.0;
+
+	/** Cull tiers that have a surface area smaller than a specified threshold. */
+	UPROPERTY(BlueprintReadWrite, EditAnywhere, Category = Settings, meta = (PCG_Overridable, InlineEditConditionToggle))
+	bool bEnableMinAreaCulling = false;
+
+	/** If a tier is smaller in area than this threshold, it will be culled. */
+	UPROPERTY(BlueprintReadWrite, EditAnywhere, Category = Settings, meta = (PCG_Overridable, EditCondition = "bEnableMinAreaCulling"))
+	double MinAreaCullingThreshold = 100.0;
+
+	/** If multiple tiers can be combined into a single tier without affecting the contour, remove the redundant one. Note: This will currently cull even if there are other unique tiers in between. */
+	UPROPERTY(BlueprintReadWrite, EditAnywhere, Category = Settings, meta = (PCG_Overridable))
+	bool bRemoveRedundantSections = true;
+
+	/** A safeguard to prevent finding features on an overly complex mesh. */
+	UPROPERTY(BlueprintReadWrite, EditAnywhere, Category = "Settings|Advanced", meta = (PCG_Overridable, EditCondition = "bEnableFindFeatures", EditConditionHides))
+	int32 MaxMeshVertexCount = 2048;
+};
+
+class FPCGSplineCrossSectionGeneratorElement : public IPCGElement
+{
+protected:
+	virtual bool ExecuteInternal(FPCGContext* Context) const override;
+	virtual bool IsCacheable(const UPCGSettings* InSettings) const override { return false; }
+};
