[PCG] Fix Mac Compile Error for Cross-Section Generator

#rnx

diff --git a/Engine/Plugins/PCGGeometryScriptInterop/Source/PCGGeometryScriptInterop/Private/Elements/PCGSplineCrossSectionGenerator.cpp b/Engine/Plugins/PCGGeometryScriptInterop/Source/PCGGeometryScriptInterop/Private/Elements/PCGSplineCrossSectionGenerator.cpp
--- a/Engine/Plugins/PCGGeometryScriptInterop/Source/PCGGeometryScriptInterop/Private/Elements/PCGSplineCrossSectionGenerator.cpp
+++ b/Engine/Plugins/PCGGeometryScriptInterop/Source/PCGGeometryScriptInterop/Private/Elements/PCGSplineCrossSectionGenerator.cpp
@@ -25,8 +25,8 @@
 namespace SplineCrossSectionGeneratorConstants
 {
 	const FName DefaultExtrusionVectorAttributeName(TEXT("ExtrusionVector"));
-	constexpr FGeometryScriptMeshPlaneCutOptions CutPlaneOptions = FGeometryScriptMeshPlaneCutOptions(/*bFillHoles=*/false, /*bFillSpans=*/false, /*bFlipCutSide=*/true, 1);
-	constexpr FGeometryScriptCopyMeshFromComponentOptions CopyMeshFromComponentOptions(/*bWantNormals=*/false, /*bWantTangents=*/false);
+	static const FGeometryScriptMeshPlaneCutOptions CutPlaneOptions{/*bFillHoles=*/false, /*bFillSpans=*/false, /*bFlipCutSide=*/true, 1};
+	static const FGeometryScriptCopyMeshFromComponentOptions CopyMeshFromComponentOptions{/*bWantNormals=*/false, /*bWantTangents=*/false};
 }
 
 struct FCrossSection
