[PCG] Fix Mesh sampling crash on Point per triangle/Vertex

#rnx

diff --git a/Engine/Plugins/PCGGeometryScriptInterop/Source/PCGGeometryScriptInterop/Private/Elements/PCGMeshSampler.cpp b/Engine/Plugins/PCGGeometryScriptInterop/Source/PCGGeometryScriptInterop/Private/Elements/PCGMeshSampler.cpp
--- a/Engine/Plugins/PCGGeometryScriptInterop/Source/PCGGeometryScriptInterop/Private/Elements/PCGMeshSampler.cpp
+++ b/Engine/Plugins/PCGGeometryScriptInterop/Source/PCGGeometryScriptInterop/Private/Elements/PCGMeshSampler.cpp
@@ -38,7 +38,7 @@ namespace PCGMeshSampler
 		auto IterationBody = [Settings, Context, SetPointDensityPtr](int32 ReadIndex, int32 WriteIndex) -> bool
 		{
 			int32 DataIndex = 0;
-			while (!Context->DynamicMeshes[DataIndex] || ReadIndex > Context->StartingIndices[DataIndex])
+			while (!Context->DynamicMeshes[DataIndex] || ReadIndex >= Context->StartingIndices[DataIndex + 1])
 			{
 				DataIndex++;
 			}
@@ -67,6 +67,12 @@ namespace PCGMeshSampler
 			return true;
 		};
 
+		// Small verification to catch any problem that could arise in the iteration body
+		if (Context->DynamicMeshes.IsEmpty() || !ensure(Context->DynamicMeshes.Num() + 1 == Context->StartingIndices.Num()))
+		{
+			return true;
+		}
+
 		return FPCGAsync::AsyncProcessingOneToOneEx(&Context->AsyncState, Context->StartingIndices.Last(), []() {}, IterationBody, /*bEnableTimeSlicing=*/true);
 	}
 
@@ -79,13 +85,14 @@ namespace PCGMeshSampler
 		auto IterationBody = [Context, Settings, SetPointDensityPtr](int32 ReadIndex, int32 WriteIndex) -> bool
 		{
 			int32 DataIndex = 0;
-			while (!Context->DynamicMeshes[DataIndex] || ReadIndex >= Context->StartingIndices[DataIndex])
+			while (!Context->DynamicMeshes[DataIndex] || ReadIndex >= Context->StartingIndices[DataIndex + 1])
 			{
 				DataIndex++;
 			}
 
 			const TArray<int32>& TriangleIds = *Context->TriangleIds[DataIndex].List.Get();
 			const int32 CurrentIndex = ReadIndex - Context->StartingIndices[DataIndex];
+
 			const int32 TriangleId = TriangleIds[CurrentIndex];
 
 			FVector Vertex1, Vertex2, Vertex3;
@@ -114,6 +121,12 @@ namespace PCGMeshSampler
 			return true;
 		};
 
+		// Small verification to catch any problem that could arise in the iteration body
+		if (Context->DynamicMeshes.IsEmpty() || !ensure(Context->DynamicMeshes.Num() + 1 == Context->StartingIndices.Num()))
+		{
+			return true;
+		}
+
 		return FPCGAsync::AsyncProcessingOneToOneEx(&Context->AsyncState, Context->StartingIndices.Last(), []() {}, IterationBody, /*bEnableTimeSlicing=*/true);
 	}
 
