[PCG] Add material extraction to the mesh sampler

#rnx

diff --git a/Engine/Plugins/PCGGeometryScriptInterop/Source/PCGGeometryScriptInterop/PCGGeometryScriptInterop.Build.cs b/Engine/Plugins/PCGGeometryScriptInterop/Source/PCGGeometryScriptInterop/PCGGeometryScriptInterop.Build.cs
--- a/Engine/Plugins/PCGGeometryScriptInterop/Source/PCGGeometryScriptInterop/PCGGeometryScriptInterop.Build.cs
+++ b/Engine/Plugins/PCGGeometryScriptInterop/Source/PCGGeometryScriptInterop/PCGGeometryScriptInterop.Build.cs
@@ -24,7 +24,8 @@ namespace UnrealBuildTool.Rules
 				new string[]
 				{
 					"GeometryCore",
-					"GeometryFramework"
+					"GeometryFramework",
+					"ModelingComponents"
 				}
 			);
 
diff --git a/Engine/Plugins/PCGGeometryScriptInterop/Source/PCGGeometryScriptInterop/Private/Elements/PCGMeshSampler.cpp b/Engine/Plugins/PCGGeometryScriptInterop/Source/PCGGeometryScriptInterop/Private/Elements/PCGMeshSampler.cpp
--- a/Engine/Plugins/PCGGeometryScriptInterop/Source/PCGGeometryScriptInterop/Private/Elements/PCGMeshSampler.cpp
+++ b/Engine/Plugins/PCGGeometryScriptInterop/Source/PCGGeometryScriptInterop/Private/Elements/PCGMeshSampler.cpp
@@ -2,6 +2,8 @@
 
 #include "Elements/PCGMeshSampler.h"
 
+#include "Helpers/PCGGeometryHelpers.h"
+
 #include "PCGComponent.h"
 #include "PCGModule.h"
 #include "PCGPin.h"
@@ -13,10 +15,12 @@
 #include "Metadata/PCGMetadataAttributeTpl.h"
 
 #include "UDynamicMesh.h"
+#include "ConversionUtils/SceneComponentToDynamicMesh.h"
 #include "Engine/AssetManager.h"
 #include "Engine/StaticMesh.h"
 #include "GameFramework/Actor.h"
 #include "GeometryScript/MeshAssetFunctions.h"
+#include "GeometryScript/MeshMaterialFunctions.h"
 #include "GeometryScript/MeshNormalsFunctions.h"
 #include "GeometryScript/MeshQueryFunctions.h"
 #include "GeometryScript/MeshRepairFunctions.h"
@@ -112,7 +116,7 @@ namespace PCGMeshSampler
 			UGeometryScriptLibrary_MeshQueryFunctions::ComputeTriangleBarycentricCoords(Context->DynamicMeshes[DataIndex], TriangleId, bIsValid, Position, Dummy1, Dummy2, Dummy3, BarycentricCoord);
 
 			Context->SetPointColorAndDensity(SetPointDensityPtr, TriangleId, BarycentricCoord, OutPoint, DataIndex);
-			Context->SetUVValueAndTriangleId(Settings->UVChannel, TriangleId, BarycentricCoord, OutPoint, DataIndex);
+			Context->SetAttributeValues(Settings->UVChannel, TriangleId, BarycentricCoord, OutPoint, DataIndex);
 
 			UPCGBlueprintHelpers::SetSeedFromPosition(OutPoint);
 
@@ -203,7 +207,7 @@ namespace PCGMeshSampler
 						OutPoint.Steepness = Settings->PointSteepness;
 
 						Context->SetPointColorAndDensity(SetPointDensityPtr, TriangleId, BarycentricCoords, OutPoint, DataIndex);
-						Context->SetUVValueAndTriangleId(Settings->UVChannel, TriangleId, BarycentricCoords, OutPoint, DataIndex);
+						Context->SetAttributeValues(Settings->UVChannel, TriangleId, BarycentricCoords, OutPoint, DataIndex);
 
 						UPCGBlueprintHelpers::SetSeedFromPosition(OutPoint);
 					}
@@ -230,7 +234,7 @@ namespace PCGMeshSampler
 	}
 }
 
-void FPCGMeshSamplerContext::SetUVValueAndTriangleId(int32 UVChannel, int32 TriangleId, const FVector& BarycentricCoord, FPCGPoint& OutPoint, int32 DataIndex)
+void FPCGMeshSamplerContext::SetAttributeValues(int32 UVChannel, int32 TriangleId, const FVector& BarycentricCoord, FPCGPoint& OutPoint, int32 DataIndex)
 {
 	if (UVAttributes.IsValidIndex(DataIndex) && UVAttributes[DataIndex])
 	{
@@ -245,6 +249,24 @@ void FPCGMeshSamplerContext::SetUVValueAndTriangleId(int32 UVChannel, int32 Tria
 		}
 	}
 
+	if (MaterialIdAttributes.IsValidIndex(DataIndex) && MaterialIdAttributes[DataIndex])
+	{
+		bool bIsValidTriangle = false;
+		const int32 MaterialId = UGeometryScriptLibrary_MeshMaterialFunctions::GetTriangleMaterialID(DynamicMeshes[DataIndex], TriangleId, bIsValidTriangle);
+
+		if (bIsValidTriangle)
+		{
+			check(OutPointData[DataIndex]);
+			OutPointData[DataIndex]->Metadata->InitializeOnSet(OutPoint.MetadataEntry);
+			MaterialIdAttributes[DataIndex]->SetValue(OutPoint.MetadataEntry, MaterialId);
+
+			if (MaterialAttributes.IsValidIndex(DataIndex) && MaterialAttributes[DataIndex] && AssetMaterialList.IsValidIndex(DataIndex) && AssetMaterialList[DataIndex].IsValidIndex(MaterialId) && AssetMaterialList[DataIndex][MaterialId])
+			{
+				MaterialAttributes[DataIndex]->SetValue(OutPoint.MetadataEntry, AssetMaterialList[DataIndex][MaterialId]);
+			}
+		}
+	}
+
 	if (TriangleIdAttributes.IsValidIndex(DataIndex) && TriangleIdAttributes[DataIndex])
 	{
 		check(OutPointData[DataIndex]);
@@ -384,6 +406,32 @@ bool FPCGMeshSamplerElement::PrepareDataInternal(FPCGContext* InContext) const
 		}
 	}
 
+	UGeometryScriptDebug* Debug = FPCGContext::NewObject_AnyThread<UGeometryScriptDebug>(Context);
+
+	// Make sure the LOD is valid, and make the conversion.
+	UE::Conversion::EMeshLODType RequestedLODType;
+
+	switch (Settings->RequestedLODType)
+	{
+	case EGeometryScriptLODType::MaxAvailable:
+		RequestedLODType = UE::Conversion::EMeshLODType::MaxAvailable;
+		break;
+	case EGeometryScriptLODType::HiResSourceModel:
+		RequestedLODType = UE::Conversion::EMeshLODType::HiResSourceModel;
+		break;
+	case EGeometryScriptLODType::SourceModel:
+		RequestedLODType = UE::Conversion::EMeshLODType::SourceModel;
+		break;
+	case EGeometryScriptLODType::RenderData:
+		RequestedLODType = UE::Conversion::EMeshLODType::RenderData;
+		break;
+	default:
+	{
+		PCGLog::LogErrorOnGraph(LOCTEXT("InvalidLOD", "Requested LOD is invalid."), Context);
+		return true;
+	}
+	}
+
 	for (const auto& [Path, DummyIndex, DummyIndex2] : Context->PathsToObjectsAndDataIndex)
 	{
 		UObject* Object = Path.ResolveObject();
@@ -393,8 +441,8 @@ bool FPCGMeshSamplerElement::PrepareDataInternal(FPCGContext* InContext) const
 			return true;
 		}
 
-		EGeometryScriptOutcomePins Outcome;
-		const FGeometryScriptMeshReadLOD MeshReadLOD{Settings->RequestedLODType, Settings->RequestedLODIndex };
+		EGeometryScriptOutcomePins Outcome = EGeometryScriptOutcomePins::Success;
+		Debug->Messages.Reset();
 		USceneComponent* SceneComponent = Cast<USceneComponent>(Object);
 
 		// @todo_pcg: Better support of multiple scene component on a given actor.
@@ -405,19 +453,45 @@ bool FPCGMeshSamplerElement::PrepareDataInternal(FPCGContext* InContext) const
 
 		if (SceneComponent)
 		{
+			// Adaptation of UGeometryScriptLibrary_SceneUtilityFunctions::CopyMeshFromComponent, since we don't have access to the Material list with this one.
 			Context->DynamicMeshes.Add(FPCGContext::NewObject_AnyThread<UDynamicMesh>(Context));
 			FTransform Transform;
-			FGeometryScriptCopyMeshFromComponentOptions Options{};
+			FText ErrorMessage;
+			UE::Conversion::FToMeshOptions Options{};
 			Options.bWantInstanceColors = true;
-			Options.RequestedLOD = MeshReadLOD;
+			Options.LODType = RequestedLODType;
+			Options.LODIndex = Settings->RequestedLODIndex;
+
+			UE::Geometry::FDynamicMesh3 TempDynMesh{};
+
+			TArray<UMaterialInterface*>* ComponentMaterialListPtr = Settings->bOutputMaterialInfo ? &Context->ComponentMaterialList.Emplace_GetRef() : nullptr;
+			TArray<UMaterialInterface*>* AssetMaterialListPtr = Settings->bOutputMaterialInfo ? &Context->AssetMaterialList.Emplace_GetRef() : nullptr;
 
-			UGeometryScriptLibrary_SceneUtilityFunctions::CopyMeshFromComponent(SceneComponent, Context->DynamicMeshes.Last(), Options, /*bTransformToWorld=*/false, Transform, Outcome);
+			const bool bSuccess = UE::Conversion::SceneComponentToDynamicMesh(SceneComponent, Options, /*bTransformToWorld=*/false, TempDynMesh, Transform, ErrorMessage, ComponentMaterialListPtr, AssetMaterialListPtr);
+			if (!bSuccess)
+			{
+				Outcome = EGeometryScriptOutcomePins::Failure;
+				FGeometryScriptDebugMessage DebugMessage;
+				DebugMessage.Message = std::move(ErrorMessage);
+				Debug->Messages.Emplace(std::move(DebugMessage));
+			}
+			else
+			{
+				Context->DynamicMeshes.Last()->SetMesh(std::move(TempDynMesh));
+			}
 		}
 		else if (UStaticMesh* StaticMesh = Cast<UStaticMesh>(Object))
 		{
+			const FGeometryScriptMeshReadLOD MeshReadLOD{ Settings->RequestedLODType, Settings->RequestedLODIndex };
+
 			Context->DynamicMeshes.Add(FPCGContext::NewObject_AnyThread<UDynamicMesh>(Context));
 			FGeometryScriptCopyMeshFromAssetOptions Options{};
-			UGeometryScriptLibrary_StaticMeshFunctions::CopyMeshFromStaticMesh(StaticMesh, Context->DynamicMeshes.Last(), Options, MeshReadLOD, Outcome);
+
+			UGeometryScriptLibrary_StaticMeshFunctions::CopyMeshFromStaticMesh(StaticMesh, Context->DynamicMeshes.Last(), Options, MeshReadLOD, Outcome, Debug);
+			if (Outcome == EGeometryScriptOutcomePins::Success && Settings->bOutputMaterialInfo)
+			{
+				UGeometryScriptLibrary_StaticMeshFunctions::GetMaterialListFromStaticMesh(StaticMesh, Context->AssetMaterialList.Emplace_GetRef());
+			}
 		}
 		else
 		{
@@ -425,14 +499,17 @@ bool FPCGMeshSamplerElement::PrepareDataInternal(FPCGContext* InContext) const
 			continue;
 		}
 
-		if (Outcome == EGeometryScriptOutcomePins::Failure)
+		if (Outcome != EGeometryScriptOutcomePins::Success)
 		{
-			PCGLog::LogErrorOnGraph(FText::Format(LOCTEXT("ObjectToDynamicMeshFailed", "Object to Dynamic mesh failed for object {0}"), FText::FromString(Path.ToString())));
+			PCGLog::LogErrorOnGraph(FText::Format(LOCTEXT("ObjectToDynamicMeshFailed", "Object to Dynamic mesh failed for object {0}."), FText::FromString(Path.ToString())));
+			PCGGeometryHelpers::GeometryScriptDebugToPCGLog(Context, Debug);
 			Context->DynamicMeshes.Last()->MarkAsGarbage();
 			Context->DynamicMeshes.RemoveAtSwap(Context->DynamicMeshes.Num() - 1);
 		}
 	}
 
+	Debug->MarkAsGarbage();
+
 	// Reserve arrays. Add one more entry for the starting indices to have the total number of items to process.
 	switch (Settings->SamplingMethod)
 	{
@@ -516,30 +593,27 @@ bool FPCGMeshSamplerElement::PrepareDataInternal(FPCGContext* InContext) const
 			CurrentOutPointData->GetMutablePoints().SetNumUninitialized(NumIterations);
 		}
 
-		// It's not clear how to compute UVs for Vertices as they are part of multiple triangles. So disable for this mode. Same for triangle ids.
-		if (Settings->SamplingMethod != EPCGMeshSamplingMethod::OnePointPerVertex)
+		auto CreateAttribute = [this, Context, Metadata = CurrentOutPointData->Metadata]<typename T>(const bool bShouldAdd, const FName AttributeName, T DefaultValue, bool bAllowInterpolation, TArray<FPCGMetadataAttribute<T>*>& OutAttributesArray, const TCHAR* What)
 		{
-			if (Settings->bExtractUVAsAttribute)
+			if (bShouldAdd)
 			{
-				FPCGMetadataAttribute<FVector2D>* UVAttribute = CurrentOutPointData->Metadata->CreateAttribute<FVector2D>(Settings->UVAttributeName, FVector2D::ZeroVector, /*bAllowsInterpolation=*/true, /*bOverrideParent=*/true);
-				if (!UVAttribute)
+				FPCGMetadataAttribute<T>* Attribute = Metadata->CreateAttribute<T>(AttributeName, std::move(DefaultValue), bAllowInterpolation, /*bOverrideParent=*/true);
+				if (!Attribute)
 				{
-					PCGE_LOG(Warning, GraphAndLog, FText::Format(LOCTEXT("AttributeUVFailed", "Failed to create attribute {0} for UVs. UVs won't be computed"), FText::FromName(Settings->UVAttributeName)));
+					PCGE_LOG(Warning, GraphAndLog, FText::Format(LOCTEXT("AttributeCreationFailed", "Failed to create attribute {0} for {1}. {1} won't be computed"), FText::FromName(AttributeName), FText::FromString(What)));
 				}
 
-				Context->UVAttributes.Emplace(UVAttribute);
+				OutAttributesArray.Emplace(Attribute);
 			}
+		};
 
-			if (Settings->bOutputTriangleIds)
-			{
-				FPCGMetadataAttribute<int32>* TriangleIdAttribute = CurrentOutPointData->Metadata->CreateAttribute<int32>(Settings->TriangleIdAttributeName, 0, /*bAllowsInterpolation=*/false, /*bOverrideParent=*/true);
-				if (!TriangleIdAttribute)
-				{
-					PCGE_LOG(Warning, GraphAndLog, FText::Format(LOCTEXT("AttributeTriangleIdFailed", "Failed to create attribute {0} for triangles ids. Triangle Ids won't be output"), FText::FromName(Settings->TriangleIdAttributeName)));
-				}
-
-				Context->TriangleIdAttributes.Emplace(TriangleIdAttribute);
-			}
+		// It's not clear how to compute UVs for Vertices as they are part of multiple triangles. So disable for this mode. Same for triangle ids.
+		if (Settings->SamplingMethod != EPCGMeshSamplingMethod::OnePointPerVertex)
+		{
+			CreateAttribute(Settings->bExtractUVAsAttribute, Settings->UVAttributeName, FVector2D::ZeroVector, /*bAllowInterpolation=*/true, Context->UVAttributes, TEXT("UVs"));
+			CreateAttribute(Settings->bOutputTriangleIds, Settings->TriangleIdAttributeName, -1, /*bAllowInterpolation=*/false, Context->TriangleIdAttributes, TEXT("Triangle IDs"));
+			CreateAttribute(Settings->bOutputMaterialInfo, Settings->MaterialIdAttributeName, -1, /*bAllowInterpolation=*/false, Context->MaterialIdAttributes, TEXT("Material IDs"));
+			CreateAttribute(Settings->bOutputMaterialInfo, Settings->MaterialAttributeName, FSoftObjectPath(), /*bAllowInterpolation=*/false, Context->MaterialAttributes, TEXT("Material"));
 		}
 
 		Outputs.Emplace_GetRef().Data = CurrentOutPointData;
diff --git a/Engine/Plugins/PCGGeometryScriptInterop/Source/PCGGeometryScriptInterop/Private/Elements/PCGSplineCrossSectionGenerator.cpp b/Engine/Plugins/PCGGeometryScriptInterop/Source/PCGGeometryScriptInterop/Private/Elements/PCGSplineCrossSectionGenerator.cpp
--- a/Engine/Plugins/PCGGeometryScriptInterop/Source/PCGGeometryScriptInterop/Private/Elements/PCGSplineCrossSectionGenerator.cpp
+++ b/Engine/Plugins/PCGGeometryScriptInterop/Source/PCGGeometryScriptInterop/Private/Elements/PCGSplineCrossSectionGenerator.cpp
@@ -2,6 +2,8 @@
 
 #include "Elements/PCGSplineCrossSectionGenerator.h"
 
+#include "Helpers/PCGGeometryHelpers.h"
+
 #include "PCGContext.h"
 #include "PCGComponent.h"
 #include "Data/PCGPrimitiveData.h"
@@ -313,10 +315,7 @@ bool FPCGSplineCrossSectionGeneratorElement::ExecuteInternal(FPCGContext* Contex
 
 #if WITH_EDITOR
 		// At this point, we're done with GeometryScript for this iteration. Print errors if they occur.
-		for (const FGeometryScriptDebugMessage& Message : DynamicMeshDebug->Messages)
-		{
-			PCGLog::LogWarningOnGraph(FText::Format(LOCTEXT("GeometryScriptError", "GeometryScript Error: {0}"), Message.Message), Context);
-		}
+		PCGGeometryHelpers::GeometryScriptDebugToPCGLog(Context, DynamicMeshDebug);
 #endif // WITH_EDITOR
 
 		if (NumLoops < 1 || PolyPaths.IsEmpty())
diff --git a/Engine/Plugins/PCGGeometryScriptInterop/Source/PCGGeometryScriptInterop/Private/Helpers/PCGGeometryHelpers.cpp b/Engine/Plugins/PCGGeometryScriptInterop/Source/PCGGeometryScriptInterop/Private/Helpers/PCGGeometryHelpers.cpp
new file mode 100644
--- /dev/null
+++ b/Engine/Plugins/PCGGeometryScriptInterop/Source/PCGGeometryScriptInterop/Private/Helpers/PCGGeometryHelpers.cpp
@@ -0,0 +1,27 @@
+// Copyright Epic Games, Inc. All Rights Reserved.
+
+#include "Helpers/PCGGeometryHelpers.h"
+
+#include "Utils/PCGLogErrors.h"
+
+#include "GeometryScript/GeometryScriptTypes.h"
+
+namespace PCGGeometryHelpers
+{
+	void GeometryScriptDebugToPCGLog(FPCGContext* Context, const UGeometryScriptDebug* Debug)
+	{
+		check(Context && Debug);
+
+		for (const FGeometryScriptDebugMessage& Message : Debug->Messages)
+		{
+			if (Message.MessageType == EGeometryScriptDebugMessageType::ErrorMessage)
+			{
+				PCGLog::LogErrorOnGraph(Message.Message, Context);
+			}
+			else
+			{
+				PCGLog::LogWarningOnGraph(Message.Message, Context);
+			}
+		}
+	}
+}
\ No newline at end of file
diff --git a/Engine/Plugins/PCGGeometryScriptInterop/Source/PCGGeometryScriptInterop/Public/Elements/PCGMeshSampler.h b/Engine/Plugins/PCGGeometryScriptInterop/Source/PCGGeometryScriptInterop/Public/Elements/PCGMeshSampler.h
--- a/Engine/Plugins/PCGGeometryScriptInterop/Source/PCGGeometryScriptInterop/Public/Elements/PCGMeshSampler.h
+++ b/Engine/Plugins/PCGGeometryScriptInterop/Source/PCGGeometryScriptInterop/Public/Elements/PCGMeshSampler.h
@@ -146,6 +146,15 @@ public:
 	UPROPERTY(BlueprintReadWrite, EditAnywhere, Category = "Settings|Extra", meta = (PCG_Overridable, EditCondition = "SamplingMethod != EPCGMeshSamplingMethod::OnePointPerVertex && bOutputTriangleIds", EditConditionHides))
 	FName TriangleIdAttributeName = TEXT("TriangleId");
 
+	UPROPERTY(BlueprintReadWrite, EditAnywhere, Category = "Settings|Extra", meta = (PCG_Overridable, EditCondition = "SamplingMethod != EPCGMeshSamplingMethod::OnePointPerVertex"))
+	bool bOutputMaterialInfo = false;
+
+	UPROPERTY(BlueprintReadWrite, EditAnywhere, Category = "Settings|Extra", meta = (PCG_Overridable, EditCondition = "SamplingMethod != EPCGMeshSamplingMethod::OnePointPerVertex && bOutputMaterialInfo", EditConditionHides))
+	FName MaterialIdAttributeName = TEXT("MaterialId");
+
+	UPROPERTY(BlueprintReadWrite, EditAnywhere, Category = "Settings|Extra", meta = (PCG_Overridable, EditCondition = "SamplingMethod != EPCGMeshSamplingMethod::OnePointPerVertex && bOutputMaterialInfo", EditConditionHides))
+	FName MaterialAttributeName = TEXT("Material");
+
 	/** Each PCG point represents a discretized, volumetric region of world space. The points' Steepness value [0.0 to
 	 * 1.0] establishes how "hard" or "soft" that volume will be represented. From 0, it will ramp up linearly
 	 * increasing its influence over the density from the point's center to up to two times the bounds. At 1, it will
@@ -178,7 +187,7 @@ public:
 
 	using SetPointDensityFunc = void(*)(const FLinearColor&, FPCGPoint&);
 
-	void SetUVValueAndTriangleId(int32 UVChannel, int32 TriangleId, const FVector& BarycentricCoord, FPCGPoint& OutPoint, int32 DataIndex);
+	void SetAttributeValues(int32 UVChannel, int32 TriangleId, const FVector& BarycentricCoord, FPCGPoint& OutPoint, int32 DataIndex);
 	void SetPointColorAndDensity(SetPointDensityFunc SetPointDensityFuncPtr, int32 TriangleId, const FVector& BarycentricCoord, FPCGPoint& OutPoint, int32 DataIndex);
 
 protected:
@@ -200,6 +209,12 @@ public:
 	// Optional attributes.
 	TArray<FPCGMetadataAttribute<FVector2D>*> UVAttributes;
 	TArray<FPCGMetadataAttribute<int32>*> TriangleIdAttributes;
+	TArray<FPCGMetadataAttribute<int32>*> MaterialIdAttributes;
+	TArray<FPCGMetadataAttribute<FSoftObjectPath>*> MaterialAttributes;
+
+	// Material Specific
+	TArray<TArray<UMaterialInterface*>> ComponentMaterialList;
+	TArray<TArray<UMaterialInterface*>> AssetMaterialList;
 
 	// For Poisson sampling, we are starting futures that are not framebound
 	// Store the futures and synchronisation items in the context
diff --git a/Engine/Plugins/PCGGeometryScriptInterop/Source/PCGGeometryScriptInterop/Public/Helpers/PCGGeometryHelpers.h b/Engine/Plugins/PCGGeometryScriptInterop/Source/PCGGeometryScriptInterop/Public/Helpers/PCGGeometryHelpers.h
new file mode 100644
--- /dev/null
+++ b/Engine/Plugins/PCGGeometryScriptInterop/Source/PCGGeometryScriptInterop/Public/Helpers/PCGGeometryHelpers.h
@@ -0,0 +1,13 @@
+// Copyright Epic Games, Inc. All Rights Reserved.
+
+#pragma once
+
+#include "Elements/PCGMeshSampler.h"
+
+struct FPCGContext;
+class UGeometryScriptDebug;
+
+namespace PCGGeometryHelpers
+{
+	PCGGEOMETRYSCRIPTINTEROP_API void GeometryScriptDebugToPCGLog(FPCGContext* Context, const UGeometryScriptDebug* Debug);
+}
