[PCG] Move all interops into a dedicated folder for tidiness purposes.

#rnx

#virtualized

diff --git a/Engine/Plugins/Experimental/PCGInterops/PCGWaterInterop/PCGWaterInterop.uplugin b/Engine/Plugins/Experimental/PCGInterops/PCGWaterInterop/PCGWaterInterop.uplugin
new file mode 100644
--- /dev/null
+++ b/Engine/Plugins/Experimental/PCGInterops/PCGWaterInterop/PCGWaterInterop.uplugin
@@ -0,0 +1,35 @@
+{
+	"FileVersion": 3,
+	"Version": 1,
+	"VersionName": "0.1",
+	"FriendlyName": "Procedural Content Generation Framework (PCG) Water Interop",
+	"Description": "Extra plugin for Procedural Content Generation Framework interacting with the Water system.",
+	"Category": "Editor",
+	"CreatedBy": "Epic Games, Inc.",
+	"CreatedByURL": "https://epicgames.com",
+	"DocsURL": "https://docs.unrealengine.com/latest/en-US/procedural-content-generation--framework-in-unreal-engine/",
+	"MarketplaceURL": "",
+	"SupportURL": "",
+	"EnabledByDefault": false,
+	"CanContainContent": true,
+	"IsBetaVersion": false,
+	"IsExperimentalVersion": true,
+	"Installed": false,
+	"Modules": [
+		{
+			"Name": "PCGWaterInterop",
+			"Type": "Runtime",
+			"LoadingPhase": "Default"
+		}
+	],
+	"Plugins": [
+		{
+			"Name": "PCG",
+			"Enabled": true
+		},
+		{
+			"Name": "Water",
+			"Enabled": true
+		}
+	]
+}
diff --git a/Engine/Plugins/Experimental/PCGInterops/PCGWaterInterop/Source/PCGWaterInterop/PCGWaterInterop.Build.cs b/Engine/Plugins/Experimental/PCGInterops/PCGWaterInterop/Source/PCGWaterInterop/PCGWaterInterop.Build.cs
new file mode 100644
--- /dev/null
+++ b/Engine/Plugins/Experimental/PCGInterops/PCGWaterInterop/Source/PCGWaterInterop/PCGWaterInterop.Build.cs
@@ -0,0 +1,22 @@
+// Copyright Epic Games, Inc. All Rights Reserved.
+
+namespace UnrealBuildTool.Rules
+{
+	public class PCGWaterInterop : ModuleRules
+	{
+		public PCGWaterInterop(ReadOnlyTargetRules Target) : base(Target)
+		{
+			PCHUsage = PCHUsageMode.UseExplicitOrSharedPCHs;
+
+			PublicDependencyModuleNames.AddRange(
+				new string[] {
+					"Core",
+					"CoreUObject",
+					"Engine",
+					"PCG",
+					"Water",
+				}
+			);
+		}
+	}
+}
diff --git a/Engine/Plugins/Experimental/PCGInterops/PCGWaterInterop/Source/PCGWaterInterop/Private/Data/PCGWaterSplineData.cpp b/Engine/Plugins/Experimental/PCGInterops/PCGWaterInterop/Source/PCGWaterInterop/Private/Data/PCGWaterSplineData.cpp
new file mode 100644
--- /dev/null
+++ b/Engine/Plugins/Experimental/PCGInterops/PCGWaterInterop/Source/PCGWaterInterop/Private/Data/PCGWaterSplineData.cpp
@@ -0,0 +1,122 @@
+// Copyright Epic Games, Inc. All Rights Reserved.
+
+#include "Data/PCGWaterSplineData.h"
+
+#include "PCGContext.h"
+#include "Data/PCGPointData.h"
+#include "Metadata/PCGMetadata.h"
+
+#include "WaterSplineComponent.h"
+#include "WaterSplineMetadata.h"
+
+#include UE_INLINE_GENERATED_CPP_BY_NAME(PCGWaterSplineData)
+
+namespace PCGWaterSplineDataConstants
+{
+	const FName DepthAttributeName = TEXT("Depth");
+	const FName WaterVelocityScalarAttributeName = TEXT("WaterVelocityScalar");
+	const FName RiverWidthAttributeName = TEXT("RiverWidth");
+	const FName AudioIntensityAttributeName = TEXT("AudioIntensity");
+}
+
+void UPCGWaterSplineData::Initialize(const UWaterSplineComponent* InSpline)
+{
+	check(InSpline && Metadata);
+
+	Super::Initialize(InSpline);
+
+	if (const UWaterSplineMetadata* WaterSplineMetadata = Cast<UWaterSplineMetadata>(InSpline->GetSplinePointsMetadata()))
+	{
+		WaterSplineMetadataStruct.Depth = WaterSplineMetadata->Depth;
+		WaterSplineMetadataStruct.WaterVelocityScalar = WaterSplineMetadata->WaterVelocityScalar;
+		WaterSplineMetadataStruct.RiverWidth = WaterSplineMetadata->RiverWidth;
+		WaterSplineMetadataStruct.AudioIntensity = WaterSplineMetadata->AudioIntensity;
+
+		Metadata->CreateAttribute<float>(PCGWaterSplineDataConstants::DepthAttributeName, InSpline->WaterSplineDefaults.DefaultDepth, /*bAllowsInterpolation=*/true, /*bOverrideParent=*/true);
+		Metadata->CreateAttribute<float>(PCGWaterSplineDataConstants::WaterVelocityScalarAttributeName, InSpline->WaterSplineDefaults.DefaultVelocity, /*bAllowsInterpolation=*/true, /*bOverrideParent=*/true);
+		Metadata->CreateAttribute<float>(PCGWaterSplineDataConstants::RiverWidthAttributeName, InSpline->WaterSplineDefaults.DefaultWidth, /*bAllowsInterpolation=*/true, /*bOverrideParent=*/true);
+		Metadata->CreateAttribute<float>(PCGWaterSplineDataConstants::AudioIntensityAttributeName, InSpline->WaterSplineDefaults.DefaultAudioIntensity, /*bAllowsInterpolation=*/true, /*bOverrideParent=*/true);
+	}
+}
+
+bool UPCGWaterSplineData::SamplePoint(const FTransform& Transform, const FBox& Bounds, FPCGPoint& OutPoint, UPCGMetadata* OutMetadata) const
+{
+	// Find nearest point on spline
+	const FVector InPosition = Transform.GetLocation();
+	const float NearestPointKey = SplineStruct.FindInputKeyClosestToWorldLocation(InPosition);
+	FTransform NearestTransform = SplineStruct.GetTransformAtSplineInputKey(NearestPointKey, ESplineCoordinateSpace::World, true);
+
+	// WaterSpline samples are scaled by river width on X instead of Y, so we swizzle XY to get the correct transform.
+	const FVector Scale = NearestTransform.GetScale3D();
+	NearestTransform.SetScale3D(FVector(Scale.Y, Scale.X, Scale.Z));
+
+	const FVector LocalPoint = NearestTransform.InverseTransformPosition(InPosition);
+	
+	// Linear fall off based on the distance to the nearest point
+	const float Distance = LocalPoint.Length();
+	if (Distance > 1.0f)
+	{
+		return false;
+	}
+	else
+	{
+		OutPoint.Transform = Transform;
+		OutPoint.SetLocalBounds(Bounds);
+		OutPoint.Density = 1.0f - Distance;
+
+		WriteMetadataToPoint(NearestPointKey, OutPoint, OutMetadata);
+
+		return true;
+	}
+}
+
+UPCGSpatialData* UPCGWaterSplineData::CopyInternal(FPCGContext* Context) const
+{
+	UPCGWaterSplineData* NewSplineData = FPCGContext::NewObject_AnyThread<UPCGWaterSplineData>(Context);
+	NewSplineData->WaterSplineMetadataStruct = WaterSplineMetadataStruct;
+
+	CopySplineData(NewSplineData);
+
+	return NewSplineData;
+}
+
+FTransform UPCGWaterSplineData::GetTransformAtDistance(int SegmentIndex, FVector::FReal Distance, bool bWorldSpace, FBox* OutBounds) const
+{
+	FTransform Transform = Super::GetTransformAtDistance(SegmentIndex, Distance, bWorldSpace, OutBounds);
+
+	// WaterSpline samples are scaled by river width on X instead of Y, so we swizzle XY to get the correct transform.
+	const FVector Scale = Transform.GetScale3D();
+	Transform.SetScale3D(FVector(Scale.Y, Scale.X, Scale.Z));
+
+	return Transform;
+}
+
+void UPCGWaterSplineData::WriteMetadataToPoint(float InputKey, FPCGPoint& OutPoint, UPCGMetadata* OutMetadata) const
+{
+	if (!OutMetadata)
+	{
+		return;
+	}
+
+	OutMetadata->InitializeOnSet(OutPoint.MetadataEntry);
+
+	if (FPCGMetadataAttribute<float>* DepthAttribute = OutMetadata->GetMutableTypedAttribute<float>(PCGWaterSplineDataConstants::DepthAttributeName))
+	{
+		DepthAttribute->SetValue(OutPoint.MetadataEntry, WaterSplineMetadataStruct.Depth.Eval(InputKey));
+	}
+
+	if (FPCGMetadataAttribute<float>* VelocityAttribute = OutMetadata->GetMutableTypedAttribute<float>(PCGWaterSplineDataConstants::WaterVelocityScalarAttributeName))
+	{
+		VelocityAttribute->SetValue(OutPoint.MetadataEntry, WaterSplineMetadataStruct.WaterVelocityScalar.Eval(InputKey));
+	}
+
+	if (FPCGMetadataAttribute<float>* WidthAttribute = OutMetadata->GetMutableTypedAttribute<float>(PCGWaterSplineDataConstants::RiverWidthAttributeName))
+	{
+		WidthAttribute->SetValue(OutPoint.MetadataEntry, WaterSplineMetadataStruct.RiverWidth.Eval(InputKey));
+	}
+
+	if (FPCGMetadataAttribute<float>* AudioIntensityAttribute = OutMetadata->GetMutableTypedAttribute<float>(PCGWaterSplineDataConstants::AudioIntensityAttributeName))
+	{
+		AudioIntensityAttribute->SetValue(OutPoint.MetadataEntry, WaterSplineMetadataStruct.AudioIntensity.Eval(InputKey));
+	}
+}
diff --git a/Engine/Plugins/Experimental/PCGInterops/PCGWaterInterop/Source/PCGWaterInterop/Private/Elements/PCGGetWaterSpline.cpp b/Engine/Plugins/Experimental/PCGInterops/PCGWaterInterop/Source/PCGWaterInterop/Private/Elements/PCGGetWaterSpline.cpp
new file mode 100644
--- /dev/null
+++ b/Engine/Plugins/Experimental/PCGInterops/PCGWaterInterop/Source/PCGWaterInterop/Private/Elements/PCGGetWaterSpline.cpp
@@ -0,0 +1,59 @@
+// Copyright Epic Games, Inc. All Rights Reserved.
+
+#include "Elements/PCGGetWaterSpline.h"
+
+#include "Data/PCGWaterSplineData.h"
+
+#include "WaterSplineComponent.h"
+
+#include UE_INLINE_GENERATED_CPP_BY_NAME(PCGGetWaterSpline)
+
+#define LOCTEXT_NAMESPACE "PCGGetWaterSplineElement"
+
+UPCGGetWaterSplineSettings::UPCGGetWaterSplineSettings()
+{
+	Mode = EPCGGetDataFromActorMode::ParseActorComponents;
+}
+
+#if WITH_EDITOR
+FText UPCGGetWaterSplineSettings::GetNodeTooltipText() const
+{
+	return LOCTEXT("GetWaterSplineTooltip", "Builds a collection of data from WaterSplineComponents on the selected actors.");
+}
+#endif
+
+FPCGElementPtr UPCGGetWaterSplineSettings::CreateElement() const
+{
+	return MakeShared<FPCGGetWaterSplineElement>();
+}
+
+void FPCGGetWaterSplineElement::ProcessActor(FPCGContext* Context, const UPCGDataFromActorSettings* Settings, AActor* FoundActor) const
+{
+	check(Context && Settings);
+
+	if (!FoundActor || !IsValid(FoundActor))
+	{
+		return;
+	}
+
+	auto NameTagsToStringTags = [](const FName& InName) { return InName.ToString(); };
+	TSet<FString> ActorTags;
+	Algo::Transform(FoundActor->Tags, ActorTags, NameTagsToStringTags);
+
+	// Prepare data on a component basis
+	TInlineComponentArray<UWaterSplineComponent*, 4> WaterSplines;
+	FoundActor->GetComponents(WaterSplines);
+
+	for (UWaterSplineComponent* WaterSplineComponent : WaterSplines)
+	{
+		UPCGWaterSplineData* SplineData = NewObject<UPCGWaterSplineData>();
+		SplineData->Initialize(WaterSplineComponent);
+
+		FPCGTaggedData& TaggedData = Context->OutputData.TaggedData.Emplace_GetRef();
+		TaggedData.Data = SplineData;
+		Algo::Transform(WaterSplineComponent->ComponentTags, TaggedData.Tags, NameTagsToStringTags);
+		TaggedData.Tags.Append(ActorTags);
+	}
+}
+
+#undef LOCTEXT_NAMESPACE
diff --git a/Engine/Plugins/Experimental/PCGInterops/PCGWaterInterop/Source/PCGWaterInterop/Private/PCGWaterInteropModule.cpp b/Engine/Plugins/Experimental/PCGInterops/PCGWaterInterop/Source/PCGWaterInterop/Private/PCGWaterInteropModule.cpp
new file mode 100644
--- /dev/null
+++ b/Engine/Plugins/Experimental/PCGInterops/PCGWaterInterop/Source/PCGWaterInterop/Private/PCGWaterInteropModule.cpp
@@ -0,0 +1,9 @@
+// Copyright Epic Games, Inc. All Rights Reserved.
+
+#include "PCGWaterInteropModule.h"
+
+#include "Modules/ModuleManager.h"
+
+IMPLEMENT_MODULE(FPCGWaterInteropModule, PCGWaterInterop);
+
+DEFINE_LOG_CATEGORY(LogPCGWaterInterop);
diff --git a/Engine/Plugins/Experimental/PCGInterops/PCGWaterInterop/Source/PCGWaterInterop/Public/Data/PCGWaterSplineData.h b/Engine/Plugins/Experimental/PCGInterops/PCGWaterInterop/Source/PCGWaterInterop/Public/Data/PCGWaterSplineData.h
new file mode 100644
--- /dev/null
+++ b/Engine/Plugins/Experimental/PCGInterops/PCGWaterInterop/Source/PCGWaterInterop/Public/Data/PCGWaterSplineData.h
@@ -0,0 +1,56 @@
+// Copyright Epic Games, Inc. All Rights Reserved.
+
+#pragma once
+
+#include "Data/PCGSplineData.h"
+
+#include "PCGWaterSplineData.generated.h"
+
+class UWaterSplineComponent;
+
+/** Helper struct to hold a copy of the WaterSplineMetadata contents. */
+USTRUCT()
+struct FPCGWaterSplineMetadataStruct
+{
+	GENERATED_BODY()
+
+	/** The Depth of the water at this vertex. */
+	UPROPERTY()
+	FInterpCurveFloat Depth;
+
+	/** The Current of the water at this vertex. Magnitude and direction. */
+	UPROPERTY()
+	FInterpCurveFloat WaterVelocityScalar;
+
+	/** Rivers Only: The width of the river (from center) in each direction. */
+	UPROPERTY()
+	FInterpCurveFloat RiverWidth;
+
+	/** A scalar used to define intensity of the water audio along the spline. */
+	UPROPERTY()
+	FInterpCurveFloat AudioIntensity;
+};
+
+UCLASS(BlueprintType, ClassGroup = (Procedural))
+class UPCGWaterSplineData : public UPCGSplineData
+{
+	GENERATED_BODY()
+
+public:
+	void Initialize(const UWaterSplineComponent* InSpline);
+
+	//~Begin UPCGSpatialData interface
+	virtual bool SamplePoint(const FTransform& Transform, const FBox& Bounds, FPCGPoint& OutPoint, UPCGMetadata* OutMetadata) const override;
+protected:
+	virtual UPCGSpatialData* CopyInternal(FPCGContext* Context) const override;
+	//~End UPCGSpatialData interface
+
+	//~Begin UPCGPolyLineData interface
+	virtual FTransform GetTransformAtDistance(int SegmentIndex, FVector::FReal Distance, bool bWorldSpace = true, FBox* OutBounds = nullptr) const override;
+	virtual void WriteMetadataToPoint(float InputKey, FPCGPoint& OutPoint, UPCGMetadata* OutMetadata) const override;
+	virtual bool HasCustomMetadata() const override { return true; }
+	//~End UPCGPolyLineData interface
+
+	UPROPERTY()
+	FPCGWaterSplineMetadataStruct WaterSplineMetadataStruct;
+};
diff --git a/Engine/Plugins/Experimental/PCGInterops/PCGWaterInterop/Source/PCGWaterInterop/Public/Elements/PCGGetWaterSpline.h b/Engine/Plugins/Experimental/PCGInterops/PCGWaterInterop/Source/PCGWaterInterop/Public/Elements/PCGGetWaterSpline.h
new file mode 100644
--- /dev/null
+++ b/Engine/Plugins/Experimental/PCGInterops/PCGWaterInterop/Source/PCGWaterInterop/Public/Elements/PCGGetWaterSpline.h
@@ -0,0 +1,36 @@
+// Copyright Epic Games, Inc. All Rights Reserved.
+
+#pragma once
+
+#include "PCGSettings.h"
+#include "Elements/PCGDataFromActor.h"
+#include "Elements/PCGTypedGetter.h"
+
+#include "PCGGetWaterSpline.generated.h"
+
+/** Builds a collection of data from WaterSplineComponents on the selected actors. */
+UCLASS(BlueprintType, ClassGroup = (Procedural))
+class UPCGGetWaterSplineSettings : public UPCGGetSplineSettings
+{
+	GENERATED_BODY()
+
+public:
+	UPCGGetWaterSplineSettings();
+
+	//~Begin UPCGSettings interface
+#if WITH_EDITOR
+	virtual FName GetDefaultNodeName() const override { return FName(TEXT("GetWaterSplineData")); }
+	virtual FText GetDefaultNodeTitle() const override { return NSLOCTEXT("PCGGetWaterSplineElement", "NodeTitle", "Get Water Spline Data"); }
+	virtual FText GetNodeTooltipText() const override;
+#endif
+
+protected:
+	virtual FPCGElementPtr CreateElement() const override;
+	//~End UPCGSettings
+};
+
+class FPCGGetWaterSplineElement : public FPCGDataFromActorElement
+{
+protected:
+	virtual void ProcessActor(FPCGContext* Context, const UPCGDataFromActorSettings* Settings, AActor* FoundActor) const override;
+};
diff --git a/Engine/Plugins/Experimental/PCGInterops/PCGWaterInterop/Source/PCGWaterInterop/Public/PCGWaterInteropModule.h b/Engine/Plugins/Experimental/PCGInterops/PCGWaterInterop/Source/PCGWaterInterop/Public/PCGWaterInteropModule.h
new file mode 100644
--- /dev/null
+++ b/Engine/Plugins/Experimental/PCGInterops/PCGWaterInterop/Source/PCGWaterInterop/Public/PCGWaterInteropModule.h
@@ -0,0 +1,16 @@
+// Copyright Epic Games, Inc. All Rights Reserved.
+
+#pragma once
+
+#include "Modules/ModuleInterface.h"
+#include "Stats/Stats.h"
+
+DECLARE_LOG_CATEGORY_EXTERN(LogPCGWaterInterop, Log, All);
+
+class FPCGWaterInteropModule final : public IModuleInterface
+{
+public:
+	// ~IModuleInterface implementation
+	virtual bool SupportsDynamicReloading() override { return true; };
+	// ~End IModuleInterface implementation
+};
diff --git a/Engine/Plugins/Experimental/PCGWaterInterop/PCGWaterInterop.uplugin b/Engine/Plugins/Experimental/PCGWaterInterop/PCGWaterInterop.uplugin
deleted file mode 100644
diff --git a/Engine/Plugins/Experimental/PCGWaterInterop/Source/PCGWaterInterop/PCGWaterInterop.Build.cs b/Engine/Plugins/Experimental/PCGWaterInterop/Source/PCGWaterInterop/PCGWaterInterop.Build.cs
deleted file mode 100644
diff --git a/Engine/Plugins/Experimental/PCGWaterInterop/Source/PCGWaterInterop/Private/Data/PCGWaterSplineData.cpp b/Engine/Plugins/Experimental/PCGWaterInterop/Source/PCGWaterInterop/Private/Data/PCGWaterSplineData.cpp
deleted file mode 100644
diff --git a/Engine/Plugins/Experimental/PCGWaterInterop/Source/PCGWaterInterop/Private/Elements/PCGGetWaterSpline.cpp b/Engine/Plugins/Experimental/PCGWaterInterop/Source/PCGWaterInterop/Private/Elements/PCGGetWaterSpline.cpp
deleted file mode 100644
diff --git a/Engine/Plugins/Experimental/PCGWaterInterop/Source/PCGWaterInterop/Private/PCGWaterInteropModule.cpp b/Engine/Plugins/Experimental/PCGWaterInterop/Source/PCGWaterInterop/Private/PCGWaterInteropModule.cpp
deleted file mode 100644
diff --git a/Engine/Plugins/Experimental/PCGWaterInterop/Source/PCGWaterInterop/Public/Data/PCGWaterSplineData.h b/Engine/Plugins/Experimental/PCGWaterInterop/Source/PCGWaterInterop/Public/Data/PCGWaterSplineData.h
deleted file mode 100644
diff --git a/Engine/Plugins/Experimental/PCGWaterInterop/Source/PCGWaterInterop/Public/Elements/PCGGetWaterSpline.h b/Engine/Plugins/Experimental/PCGWaterInterop/Source/PCGWaterInterop/Public/Elements/PCGGetWaterSpline.h
deleted file mode 100644
diff --git a/Engine/Plugins/Experimental/PCGWaterInterop/Source/PCGWaterInterop/Public/PCGWaterInteropModule.h b/Engine/Plugins/Experimental/PCGWaterInterop/Source/PCGWaterInterop/Public/PCGWaterInteropModule.h
deleted file mode 100644
diff --git a/Engine/Plugins/PCGExternalDataInterop/PCGExternalDataInterop.uplugin b/Engine/Plugins/PCGExternalDataInterop/PCGExternalDataInterop.uplugin
deleted file mode 100644
diff --git a/Engine/Plugins/PCGExternalDataInterop/Source/PCGExternalDataInterop/PCGExternalDataInterop.Build.cs b/Engine/Plugins/PCGExternalDataInterop/Source/PCGExternalDataInterop/PCGExternalDataInterop.Build.cs
deleted file mode 100644
diff --git a/Engine/Plugins/PCGExternalDataInterop/Source/PCGExternalDataInterop/Private/Alembic/PCGAlembicInterop.cpp b/Engine/Plugins/PCGExternalDataInterop/Source/PCGExternalDataInterop/Private/Alembic/PCGAlembicInterop.cpp
deleted file mode 100644
diff --git a/Engine/Plugins/PCGExternalDataInterop/Source/PCGExternalDataInterop/Private/Alembic/PCGAlembicInterop.h b/Engine/Plugins/PCGExternalDataInterop/Source/PCGExternalDataInterop/Private/Alembic/PCGAlembicInterop.h
deleted file mode 100644
diff --git a/Engine/Plugins/PCGExternalDataInterop/Source/PCGExternalDataInterop/Private/Elements/PCGLoadAlembicElement.cpp b/Engine/Plugins/PCGExternalDataInterop/Source/PCGExternalDataInterop/Private/Elements/PCGLoadAlembicElement.cpp
deleted file mode 100644
diff --git a/Engine/Plugins/PCGExternalDataInterop/Source/PCGExternalDataInterop/Private/PCGExternalDataInteropModule.cpp b/Engine/Plugins/PCGExternalDataInterop/Source/PCGExternalDataInterop/Private/PCGExternalDataInteropModule.cpp
deleted file mode 100644
diff --git a/Engine/Plugins/PCGExternalDataInterop/Source/PCGExternalDataInterop/Public/Elements/PCGLoadAlembicElement.h b/Engine/Plugins/PCGExternalDataInterop/Source/PCGExternalDataInterop/Public/Elements/PCGLoadAlembicElement.h
deleted file mode 100644
diff --git a/Engine/Plugins/PCGExternalDataInterop/Source/PCGExternalDataInteropEditor/PCGExternalDataInteropEditor.Build.cs b/Engine/Plugins/PCGExternalDataInterop/Source/PCGExternalDataInteropEditor/PCGExternalDataInteropEditor.Build.cs
deleted file mode 100644
diff --git a/Engine/Plugins/PCGExternalDataInterop/Source/PCGExternalDataInteropEditor/Private/PCGExternalDataInteropEditorModule.cpp b/Engine/Plugins/PCGExternalDataInterop/Source/PCGExternalDataInteropEditor/Private/PCGExternalDataInteropEditorModule.cpp
deleted file mode 100644
diff --git a/Engine/Plugins/PCGExternalDataInterop/Source/PCGExternalDataInteropEditor/Private/PCGLoadAlembic.cpp b/Engine/Plugins/PCGExternalDataInterop/Source/PCGExternalDataInteropEditor/Private/PCGLoadAlembic.cpp
deleted file mode 100644
diff --git a/Engine/Plugins/PCGExternalDataInterop/Source/PCGExternalDataInteropEditor/Public/PCGLoadAlembic.h b/Engine/Plugins/PCGExternalDataInterop/Source/PCGExternalDataInteropEditor/Public/PCGLoadAlembic.h
deleted file mode 100644
diff --git a/Engine/Plugins/PCGGeometryScriptInterop/PCGGeometryScriptInterop.uplugin b/Engine/Plugins/PCGGeometryScriptInterop/PCGGeometryScriptInterop.uplugin
deleted file mode 100644
diff --git a/Engine/Plugins/PCGGeometryScriptInterop/Source/PCGGeometryScriptInterop/PCGGeometryScriptInterop.Build.cs b/Engine/Plugins/PCGGeometryScriptInterop/Source/PCGGeometryScriptInterop/PCGGeometryScriptInterop.Build.cs
deleted file mode 100644
diff --git a/Engine/Plugins/PCGGeometryScriptInterop/Source/PCGGeometryScriptInterop/Private/Elements/PCGMeshSampler.cpp b/Engine/Plugins/PCGGeometryScriptInterop/Source/PCGGeometryScriptInterop/Private/Elements/PCGMeshSampler.cpp
deleted file mode 100644
diff --git a/Engine/Plugins/PCGGeometryScriptInterop/Source/PCGGeometryScriptInterop/Private/Elements/PCGSplineCrossSectionGenerator.cpp b/Engine/Plugins/PCGGeometryScriptInterop/Source/PCGGeometryScriptInterop/Private/Elements/PCGSplineCrossSectionGenerator.cpp
deleted file mode 100644
diff --git a/Engine/Plugins/PCGGeometryScriptInterop/Source/PCGGeometryScriptInterop/Private/Helpers/PCGGeometryHelpers.cpp b/Engine/Plugins/PCGGeometryScriptInterop/Source/PCGGeometryScriptInterop/Private/Helpers/PCGGeometryHelpers.cpp
deleted file mode 100644
diff --git a/Engine/Plugins/PCGGeometryScriptInterop/Source/PCGGeometryScriptInterop/Private/PCGGeometryScriptInteropModule.cpp b/Engine/Plugins/PCGGeometryScriptInterop/Source/PCGGeometryScriptInterop/Private/PCGGeometryScriptInteropModule.cpp
deleted file mode 100644
diff --git a/Engine/Plugins/PCGGeometryScriptInterop/Source/PCGGeometryScriptInterop/Public/Elements/PCGMeshSampler.h b/Engine/Plugins/PCGGeometryScriptInterop/Source/PCGGeometryScriptInterop/Public/Elements/PCGMeshSampler.h
deleted file mode 100644
diff --git a/Engine/Plugins/PCGGeometryScriptInterop/Source/PCGGeometryScriptInterop/Public/Elements/PCGSplineCrossSectionGenerator.h b/Engine/Plugins/PCGGeometryScriptInterop/Source/PCGGeometryScriptInterop/Public/Elements/PCGSplineCrossSectionGenerator.h
deleted file mode 100644
diff --git a/Engine/Plugins/PCGGeometryScriptInterop/Source/PCGGeometryScriptInterop/Public/Helpers/PCGGeometryHelpers.h b/Engine/Plugins/PCGGeometryScriptInterop/Source/PCGGeometryScriptInterop/Public/Helpers/PCGGeometryHelpers.h
deleted file mode 100644
diff --git a/Engine/Plugins/PCGInterops/PCGExternalDataInterop/PCGExternalDataInterop.uplugin b/Engine/Plugins/PCGInterops/PCGExternalDataInterop/PCGExternalDataInterop.uplugin
new file mode 100644
--- /dev/null
+++ b/Engine/Plugins/PCGInterops/PCGExternalDataInterop/PCGExternalDataInterop.uplugin
@@ -0,0 +1,39 @@
+{
+	"FileVersion": 3,
+	"Version": 1,
+	"VersionName": "0.2",
+	"FriendlyName": "Procedural Content Generation Framework (PCG) External Data Interop",
+	"Description": "Extra plugin for Procedural Content Generation Framework interacting with external data formats.",
+	"Category": "Editor",
+	"CreatedBy": "Epic Games, Inc.",
+	"CreatedByURL": "https://epicgames.com",
+	"DocsURL": "https://docs.unrealengine.com/latest/en-US/procedural-content-generation--framework-in-unreal-engine/",
+	"MarketplaceURL": "",
+	"SupportURL": "",
+	"EnabledByDefault": false,
+	"CanContainContent": true,
+	"IsBetaVersion": true,
+	"Installed": false,
+	"Modules": [
+		{
+			"Name": "PCGExternalDataInterop",
+			"Type": "Runtime",
+			"LoadingPhase": "Default"
+		},
+		{
+			"Name": "PCGExternalDataInteropEditor",
+			"Type": "Editor",
+			"LoadingPhase": "Default"
+		},
+	],
+	"Plugins": [
+		{
+			"Name": "PCG",
+			"Enabled": true
+		},
+		{
+			"Name": "AlembicImporter",
+			"Enabled": true
+		}
+	]
+}
diff --git a/Engine/Plugins/PCGInterops/PCGExternalDataInterop/Source/PCGExternalDataInterop/PCGExternalDataInterop.Build.cs b/Engine/Plugins/PCGInterops/PCGExternalDataInterop/Source/PCGExternalDataInterop/PCGExternalDataInterop.Build.cs
new file mode 100644
--- /dev/null
+++ b/Engine/Plugins/PCGInterops/PCGExternalDataInterop/Source/PCGExternalDataInterop/PCGExternalDataInterop.Build.cs
@@ -0,0 +1,33 @@
+// Copyright Epic Games, Inc. All Rights Reserved.
+
+namespace UnrealBuildTool.Rules
+{
+	public class PCGExternalDataInterop : ModuleRules
+	{
+		public PCGExternalDataInterop(ReadOnlyTargetRules Target) : base(Target)
+		{
+			PCHUsage = PCHUsageMode.UseExplicitOrSharedPCHs;
+
+			PublicDependencyModuleNames.AddRange(
+				new string[] {
+					"Core",
+					"CoreUObject",
+					"Engine",
+					"Projects",
+					"PCG"
+				});
+
+			if (Target.bBuildEditor == true)
+			{
+				PrivateDependencyModuleNames.AddRange(
+					new string[]
+					{
+						"UnrealEd",
+						"AlembicLib",
+						"AlembicLibrary"
+					}
+				);
+			}
+		}
+	}
+}
diff --git a/Engine/Plugins/PCGInterops/PCGExternalDataInterop/Source/PCGExternalDataInterop/Private/Alembic/PCGAlembicInterop.cpp b/Engine/Plugins/PCGInterops/PCGExternalDataInterop/Source/PCGExternalDataInterop/Private/Alembic/PCGAlembicInterop.cpp
new file mode 100644
--- /dev/null
+++ b/Engine/Plugins/PCGInterops/PCGExternalDataInterop/Source/PCGExternalDataInterop/Private/Alembic/PCGAlembicInterop.cpp
@@ -0,0 +1,385 @@
+// Copyright Epic Games, Inc. All Rights Reserved.
+
+#include "PCGAlembicInterop.h"
+
+#include "Elements/PCGLoadAlembicElement.h"
+
+#include "PCGComponent.h"
+#include "Data/PCGPointData.h"
+#include "Elements/Metadata/PCGMetadataElementCommon.h"
+#include "Metadata/Accessors/IPCGAttributeAccessorTpl.h"
+#include "Metadata/Accessors/PCGAttributeAccessorHelpers.h"
+#include "Metadata/Accessors/PCGAttributeAccessorKeys.h"
+
+#include <string>
+
+#define LOCTEXT_NAMESPACE "PCGExternalData_Alembic"
+
+namespace PCGAlembicInterop
+{
+
+#if WITH_EDITOR
+class FPCGAlembicPositionsAccessor : public IPCGAttributeAccessorT<FPCGAlembicPositionsAccessor>
+{
+public:
+	using Type = FVector;
+	using Super = IPCGAttributeAccessorT<FPCGAlembicPositionsAccessor>;
+
+	FPCGAlembicPositionsAccessor(Alembic::AbcGeom::IPoints::schema_type::Sample& Sample)
+		: Super(/*bInReadOnly=*/true)
+		, SamplePtr(Sample.getPositions())
+	{}
+
+	bool GetRangeImpl(TArrayView<FVector> OutValues, int32 Index, const IPCGAttributeAccessorKeys&) const
+	{
+		for (int32 i = 0; i < OutValues.Num(); ++i)
+		{
+			Alembic::Abc::P3fArraySample::value_type Position = (*SamplePtr)[Index + i];
+			OutValues[i] = FVector(Position.x, Position.y, Position.z);
+		}
+
+		return true;
+	}
+
+	bool SetRangeImpl(TArrayView<const FVector>, int32, IPCGAttributeAccessorKeys&, EPCGAttributeAccessorFlags)
+	{
+		return false;
+	}
+
+private:
+	Alembic::Abc::P3fArraySamplePtr SamplePtr;
+};
+
+template<typename T, typename AbcParamType, int32 Extent, int32 SubExtent, bool bUseCStr>
+class FPCGAlembicAccessor : public IPCGAttributeAccessorT<FPCGAlembicAccessor<T, AbcParamType, Extent, SubExtent, bUseCStr>>
+{
+public:
+	using Type = T;
+	using Super = IPCGAttributeAccessorT<FPCGAlembicAccessor<T, AbcParamType, Extent, SubExtent, bUseCStr>>;
+
+	FPCGAlembicAccessor(Alembic::AbcGeom::ICompoundProperty& Parameters, const FString& PropName)
+		: Super(/*bInReadOnly=*/true)
+		, Param(Parameters, std::string(TCHAR_TO_UTF8(*PropName)))
+		, SamplePtr(Param.getExpandedValue().getVals())
+	{}
+
+	bool GetRangeImpl(TArrayView<T> OutValues, int32 Index, const IPCGAttributeAccessorKeys&) const
+	{
+		for (int32 i = 0; i < OutValues.Num(); ++i)
+		{
+			T& OutValue = OutValues[i];
+			const int32 DataIndex = Index + i;
+
+			if constexpr (bUseCStr && Extent == 1 && SubExtent == 1) // String
+			{
+				OutValue = (*SamplePtr)[DataIndex].c_str();
+			}
+			else if constexpr (Extent == 1 && SubExtent == 1) // Scalar
+			{
+				OutValue = (*SamplePtr)[DataIndex];
+			}
+			else if constexpr (Extent > 1 && SubExtent == 1)
+			{
+				for (int32 D = 0; D < Extent; ++D)
+				{
+					OutValue[D] = (*SamplePtr)[DataIndex][D];
+				}
+			}
+			else if constexpr (Extent == 1 && SubExtent > 1)
+			{
+				for (int32 D = 0; D < SubExtent; ++D)
+				{
+					OutValue[D] = (*SamplePtr)[DataIndex * SubExtent + D];
+				}
+			}
+			else
+			{
+				return false;
+			}
+		}
+
+		return true;
+	}
+
+	bool SetRangeImpl(TArrayView<const T>, int32, IPCGAttributeAccessorKeys&, EPCGAttributeAccessorFlags)
+	{
+		return false;
+	}
+
+private:
+	AbcParamType Param;
+	using AbcSamplerType = typename AbcParamType::Sample::samp_ptr_type;
+	AbcSamplerType SamplePtr;
+};
+
+TUniquePtr<const IPCGAttributeAccessor> CreateAlembicPropAccessor(Alembic::AbcGeom::ICompoundProperty& Parameters, Alembic::Abc::PropertyHeader& PropertyHeader, const FString& PropName)
+{
+	Alembic::Abc::PropertyType PropType = PropertyHeader.getPropertyType();
+
+	Alembic::Abc::DataType DataType = PropertyHeader.getDataType();
+	int TypeExtent = DataType.getExtent();
+
+	FString MetadataExtent(PropertyHeader.getMetaData().get("arrayExtent").c_str());
+	const int32 SubExtent = (MetadataExtent.Compare("") == 0 ? 1 : FCString::Atoi(*MetadataExtent));
+
+	if (DataType.getPod() == Alembic::Util::kFloat32POD && (TypeExtent == 1 || SubExtent == 1))
+	{
+		switch (TypeExtent)
+		{
+		case 1:
+			switch (SubExtent)
+			{
+			case 1: return MakeUnique<FPCGAlembicAccessor<float, Alembic::AbcGeom::IFloatGeomParam, 1, 1, false>>(Parameters, PropName);
+			case 2: return MakeUnique<FPCGAlembicAccessor<FVector2D, Alembic::AbcGeom::IFloatGeomParam, 1, 2, false>>(Parameters, PropName);
+			case 3: return MakeUnique<FPCGAlembicAccessor<FVector, Alembic::AbcGeom::IFloatGeomParam, 1, 3, false>>(Parameters, PropName);
+			case 4: return MakeUnique<FPCGAlembicAccessor<FVector4, Alembic::AbcGeom::IFloatGeomParam, 1, 4, false>>(Parameters, PropName);
+			}
+		case 2: return MakeUnique<FPCGAlembicAccessor<FVector2D, Alembic::AbcGeom::IV2fGeomParam, 2, 1, false>>(Parameters, PropName);
+		case 3: return MakeUnique<FPCGAlembicAccessor<FVector, Alembic::AbcGeom::IV3fGeomParam, 3, 1, false>>(Parameters, PropName);
+		case 4: return MakeUnique<FPCGAlembicAccessor<FVector4, Alembic::AbcGeom::IQuatfGeomParam, 4, 1, false>>(Parameters, PropName);
+		}
+	}
+	else if (DataType.getPod() == Alembic::Util::kFloat64POD && (TypeExtent == 1 || SubExtent == 1))
+	{
+		switch (TypeExtent)
+		{
+		case 1:
+			switch (SubExtent)
+			{
+			case 1: return MakeUnique<FPCGAlembicAccessor<double, Alembic::AbcGeom::IDoubleGeomParam, 1, 1, false>>(Parameters, PropName);
+			case 2: return MakeUnique<FPCGAlembicAccessor<FVector2D, Alembic::AbcGeom::IDoubleGeomParam, 1, 2, false>>(Parameters, PropName);
+			case 3: return MakeUnique<FPCGAlembicAccessor<FVector, Alembic::AbcGeom::IDoubleGeomParam, 1, 3, false>>(Parameters, PropName);
+			case 4: return MakeUnique<FPCGAlembicAccessor<FVector4, Alembic::AbcGeom::IDoubleGeomParam, 1, 4, false>>(Parameters, PropName);
+			}
+		case 2: return MakeUnique<FPCGAlembicAccessor<FVector2D, Alembic::AbcGeom::IV2dGeomParam, 2, 1, false>>(Parameters, PropName);
+		case 3: return MakeUnique<FPCGAlembicAccessor<FVector, Alembic::AbcGeom::IV3dGeomParam, 3, 1, false>>(Parameters, PropName);
+		case 4: return MakeUnique<FPCGAlembicAccessor<FVector4, Alembic::AbcGeom::IQuatdGeomParam, 4, 1, false>>(Parameters, PropName);
+		}
+	}
+	else if ((TypeExtent == 0 || TypeExtent == 1) && SubExtent == 1) // Scalar types
+	{
+		switch (DataType.getPod())
+		{
+		case Alembic::Util::kBooleanPOD: return MakeUnique<FPCGAlembicAccessor<bool, Alembic::AbcGeom::IBoolGeomParam, 1, 1, false>>(Parameters, PropName);
+		case Alembic::Util::kInt8POD: return MakeUnique<FPCGAlembicAccessor<int32, Alembic::AbcGeom::ICharGeomParam, 1, 1, false>>(Parameters, PropName);
+		case Alembic::Util::kInt16POD: return MakeUnique<FPCGAlembicAccessor<int32, Alembic::AbcGeom::IInt16GeomParam, 1, 1, false>>(Parameters, PropName);
+		case Alembic::Util::kInt32POD: return MakeUnique<FPCGAlembicAccessor<int32, Alembic::AbcGeom::IInt32GeomParam, 1, 1, false>>(Parameters, PropName);
+		case Alembic::Util::kInt64POD: return MakeUnique<FPCGAlembicAccessor<int64, Alembic::AbcGeom::IInt64GeomParam, 1, 1, false>>(Parameters, PropName);
+		case Alembic::Util::kUnknownPOD: // fall-through
+		case Alembic::Util::kStringPOD: return MakeUnique<FPCGAlembicAccessor<FString, Alembic::AbcGeom::IStringGeomParam, 1, 1, true>>(Parameters, PropName);
+		}
+	}
+
+	return nullptr;
+}
+
+bool CreatePointAccessorAndValidate(FPCGContext* Context, UPCGPointData* PointData, const TUniquePtr<const IPCGAttributeAccessor>& AlembicPropAccessor, const FPCGAttributePropertySelector& PointPropertySelector, const FString& PropName, TUniquePtr<IPCGAttributeAccessor>& PointPropertyAccessor)
+{
+	if (!PointData || !PointData->Metadata)
+	{
+		return false;
+	}
+
+	if (!AlembicPropAccessor)
+	{
+		PCGE_LOG_C(Error, GraphAndLog, Context, FText::Format(LOCTEXT("AlembicPropertyNotSupported", "Property '{0}' is not of a supported type."), FText::FromString(PropName)));
+		return false;
+	}
+
+	UPCGMetadata* PointMetadata = PointData->Metadata;
+
+	// Create attribute if needed
+	if (PointPropertySelector.GetSelection() == EPCGAttributePropertySelection::Attribute && !PointMetadata->HasAttribute(FName(PropName)))
+	{
+		auto CreateAttribute = [PointMetadata, &PropName](auto Dummy)
+		{
+			using AttributeType = decltype(Dummy);
+			return PCGMetadataElementCommon::ClearOrCreateAttribute<AttributeType>(PointMetadata, FName(PropName)) != nullptr;
+		};
+
+		if (!PCGMetadataAttribute::CallbackWithRightType(AlembicPropAccessor->GetUnderlyingType(), CreateAttribute))
+		{
+			PCGE_LOG_C(Warning, GraphAndLog, Context, FText::Format(LOCTEXT("FailedToCreateNewAttribute", "Failed to create new attribute '{0}'"), FText::FromString(PropName)));
+			return false;
+		}
+	}
+
+	PointPropertyAccessor = PCGAttributeAccessorHelpers::CreateAccessor(PointData, PointPropertySelector);
+
+	if (!PointPropertyAccessor)
+	{
+		PCGE_LOG_C(Error, GraphAndLog, Context, FText::Format(LOCTEXT("AlembicTargetNotSupported", "Unable to write to target property '{0}.'"), PointPropertySelector.GetDisplayText()));
+		return false;
+	}
+
+	// Final verification, if we can put the value of input into output
+	if (!PCG::Private::IsBroadcastable(AlembicPropAccessor->GetUnderlyingType(), PointPropertyAccessor->GetUnderlyingType()))
+	{
+		FText InputTypeName = FText::FromString(PCG::Private::GetTypeName(AlembicPropAccessor->GetUnderlyingType()));
+		FText OutputTypeName = FText::FromString(PCG::Private::GetTypeName(PointPropertyAccessor->GetUnderlyingType()));
+
+		PCGE_LOG_C(Error, GraphAndLog, Context, FText::Format(LOCTEXT("CannotBroadcastTypes", "Cannot convert input type '{0}' into output type '{1}'"), InputTypeName, OutputTypeName));
+		return false;
+	}
+
+	return true;
+}
+
+bool ParseAlembicObject(FPCGExternalDataContext* Context, const Alembic::Abc::IObject& Object)
+{
+	check(Context);
+	const UPCGExternalDataSettings* Settings = Context->GetInputSettings<UPCGExternalDataSettings>();
+	check(Settings);
+
+	const Alembic::Abc::MetaData& ObjectMetaData = Object.getMetaData();
+	const uint32 NumChildren = Object.getNumChildren();
+
+	bool bHasCreatedData = false;
+
+	if (Alembic::AbcGeom::IPoints::matches(ObjectMetaData))
+	{
+		Alembic::AbcGeom::IPoints Points = Alembic::AbcGeom::IPoints(Object, Alembic::Abc::kWrapExisting);
+		Alembic::AbcGeom::IPoints::schema_type::Sample Sample = Points.getSchema().getValue();
+
+		Alembic::Abc::P3fArraySamplePtr Positions = Sample.getPositions();
+		uint32 NumPoints = Positions ? Positions->size() : 0;
+
+		if (NumPoints > 0)
+		{
+			bHasCreatedData = true;
+
+			// Create point data & mapping
+			UPCGPointData* PointData = NewObject<UPCGPointData>();
+			check(PointData);
+
+			UPCGMetadata* PointMetadata = PointData->MutableMetadata();
+
+			FPCGExternalDataContext::FPointDataAccessorsMapping& PointDataAccessorMapping = Context->PointDataAccessorsMapping.Emplace_GetRef();
+			PointDataAccessorMapping.Data = PointData;
+			PointDataAccessorMapping.Metadata = PointMetadata;
+			// We're not going to use the input keys, but we still need to provide something
+			PointDataAccessorMapping.RowKeys = MakeUnique<FPCGAttributeAccessorKeysEntries>(PCGInvalidEntryKey);
+
+			TArray<FPCGPoint>& OutPoints = PointData->GetMutablePoints();
+			OutPoints.SetNum(NumPoints);
+
+			// If the user has provided a position remapping, don't process points right now, instead push the transformation to the row accessors
+			if (const FPCGAttributePropertySelector* RemappedPositions = Settings->AttributeMapping.Find(TEXT("Position")))
+			{
+				TUniquePtr<const IPCGAttributeAccessor> AlembicPositionAccessor = MakeUnique<FPCGAlembicPositionsAccessor>(Sample);
+				TUniquePtr<IPCGAttributeAccessor> PointPropertyAccessor;
+
+				FPCGAttributePropertyOutputSelector PointPositionSelector;
+				PointPositionSelector.ImportFromOtherSelector(*RemappedPositions);
+				const FString PropName = PointPositionSelector.GetName().ToString();
+
+				if (CreatePointAccessorAndValidate(Context, PointData, AlembicPositionAccessor, PointPositionSelector, PropName, PointPropertyAccessor))
+				{
+					PointDataAccessorMapping.RowToPointAccessors.Emplace(MoveTemp(AlembicPositionAccessor), MoveTemp(PointPropertyAccessor), PointPositionSelector);
+				}
+			}
+			else // Otherwise, write the positions directly
+			{
+				for (uint32 PointIndex = 0; PointIndex < NumPoints; ++PointIndex)
+				{
+					Alembic::Abc::P3fArraySample::value_type Position = (*Positions)[PointIndex];
+					OutPoints[PointIndex].Transform.SetLocation(FVector(Position.x, Position.y, Position.z));
+				}
+			}
+
+			Alembic::AbcGeom::ICompoundProperty Parameters = Points.getSchema().getArbGeomParams();
+			if (Parameters.valid())
+			{
+				for (int Index = 0; Index < Parameters.getNumProperties(); ++Index)
+				{
+					Alembic::Abc::PropertyHeader PropertyHeader = Parameters.getPropertyHeader(Index);
+					FString PropName(PropertyHeader.getName().c_str());
+
+					// We'll parse only properties that affect every point/object
+					if (PropertyHeader.getPropertyType() != Alembic::Abc::kArrayProperty &&
+						!(PropertyHeader.getPropertyType() == Alembic::Abc::kCompoundProperty && (PropertyHeader.getDataType().getPod() == Alembic::Util::kUnknownPOD || PropertyHeader.getDataType().getPod() == Alembic::Util::kStringPOD)))
+					{
+						PCGE_LOG_C(Error, GraphAndLog, Context, FText::Format(LOCTEXT("AlembicPropertyType", "Property '{0}' is not supported (expected array)."), FText::FromString(PropName)));
+						continue;
+					}
+
+					TUniquePtr<const IPCGAttributeAccessor> AlembicPropAccessor = CreateAlembicPropAccessor(Parameters, PropertyHeader, PropName);
+					TUniquePtr<IPCGAttributeAccessor> PointPropertyAccessor;
+
+					// Setup attribute property selector
+					FPCGAttributePropertySelector PointPropertySelector;
+					if (const FPCGAttributePropertySelector* MappedField = Settings->AttributeMapping.Find(PropName))
+					{
+						PointPropertySelector = *MappedField;
+						PropName = PointPropertySelector.GetName().ToString();
+					}
+					else
+					{
+						PointPropertySelector.Update(PropName);
+					}
+
+					if (CreatePointAccessorAndValidate(Context, PointData, AlembicPropAccessor, PointPropertySelector, PropName, PointPropertyAccessor))
+					{
+						PointDataAccessorMapping.RowToPointAccessors.Emplace(MoveTemp(AlembicPropAccessor), MoveTemp(PointPropertyAccessor), PointPropertySelector);
+					}
+				}
+			}
+		}
+	}
+
+	if (NumChildren > 0)
+	{
+		for (uint32 ChildIndex = 0; ChildIndex < NumChildren; ++ChildIndex)
+		{
+			bHasCreatedData |= ParseAlembicObject(Context, Object.getChild(ChildIndex));
+		}
+	}
+
+	return bHasCreatedData;
+}
+
+void LoadFromAlembicFile(FPCGLoadAlembicContext* Context, const FString& FileName)
+{
+	check(Context);
+
+	/** Factory used to generate objects*/
+	Alembic::AbcCoreFactory::IFactory& Factory = Context->Factory;
+	Alembic::AbcCoreFactory::IFactory::CoreType CompressionType = Alembic::AbcCoreFactory::IFactory::kUnknown;
+	/** Archive-typed ABC file */
+	Alembic::Abc::IArchive& Archive = Context->Archive;
+	/** Alembic typed root (top) object*/
+	Alembic::Abc::IObject& TopObject = Context->TopObject;
+
+	Factory.setPolicy(Alembic::Abc::ErrorHandler::kQuietNoopPolicy);
+	const size_t ReasonableNumStreams = 12; // This is what we had in Rule Processor historically
+	Factory.setOgawaNumStreams(ReasonableNumStreams);
+
+	// Extract Archive and compression type from file
+	Archive = Factory.getArchive(TCHAR_TO_UTF8(*FileName), CompressionType);
+
+	if (!Archive.valid())
+	{
+		PCGE_LOG_C(Error, GraphAndLog, Context, FText::Format(LOCTEXT("FailedToOpenAlembic", "Failed to open '{0}': Not a valid Alembic file."), FText::FromString(FileName)));
+		return;
+	}
+
+	// Get Top/root object
+	TopObject = Alembic::Abc::IObject(Archive, Alembic::Abc::kTop);
+	if (!TopObject.valid())
+	{
+		PCGE_LOG_C(Error, GraphAndLog, Context, FText::Format(LOCTEXT("FailedToOpenAlembicRoot", "Failed to import '{0}': Alembic root is not valid."), FText::FromString(FileName)));
+		return;
+	}
+
+	if (!ParseAlembicObject(Context, TopObject))
+	{
+		PCGE_LOG_C(Warning, GraphAndLog, Context, FText::Format(LOCTEXT("FailedToCreatePointData", "Import of '{0}' was successful but there is no valid point cloud data in the Alembic file."), FText::FromString(FileName)));
+		return;
+	}
+}
+
+#endif // WITH_EDITOR
+
+}
+
+#undef LOCTEXT_NAMESPACE
\ No newline at end of file
diff --git a/Engine/Plugins/PCGInterops/PCGExternalDataInterop/Source/PCGExternalDataInterop/Private/Alembic/PCGAlembicInterop.h b/Engine/Plugins/PCGInterops/PCGExternalDataInterop/Source/PCGExternalDataInterop/Private/Alembic/PCGAlembicInterop.h
new file mode 100644
--- /dev/null
+++ b/Engine/Plugins/PCGInterops/PCGExternalDataInterop/Source/PCGExternalDataInterop/Private/Alembic/PCGAlembicInterop.h
@@ -0,0 +1,40 @@
+// Copyright Epic Games, Inc. All Rights Reserved.
+
+#pragma once
+
+#include "Elements/IO/PCGExternalDataContext.h"
+
+#if WITH_EDITOR
+#include "AbcImportSettings.h"
+#include "AbcUtilities.h"
+
+THIRD_PARTY_INCLUDES_START
+#pragma warning(push)
+#pragma warning(disable:4005) // TEXT macro redefinition
+#include <Alembic/AbcCoreOgawa/All.h>
+#include "Alembic/AbcGeom/All.h"
+#include "Alembic/AbcCoreFactory/IFactory.h"
+#pragma warning(pop)
+THIRD_PARTY_INCLUDES_END
+
+#endif // WITH_EDITOR
+
+struct FPCGLoadAlembicContext : public FPCGExternalDataContext
+{
+#if WITH_EDITOR
+	// Used for alembic parsing
+	/** Factory used to generate objects*/
+	Alembic::AbcCoreFactory::IFactory Factory;
+	/** Archive-typed ABC file */
+	Alembic::Abc::IArchive Archive;
+	/** Alembic typed root (top) object*/
+	Alembic::Abc::IObject TopObject;
+#endif // WITH_EDITOR
+};
+
+namespace PCGAlembicInterop
+{
+#if WITH_EDITOR
+	void LoadFromAlembicFile(FPCGLoadAlembicContext* Context, const FString& FileName);
+#endif
+}
\ No newline at end of file
diff --git a/Engine/Plugins/PCGInterops/PCGExternalDataInterop/Source/PCGExternalDataInterop/Private/Elements/PCGLoadAlembicElement.cpp b/Engine/Plugins/PCGInterops/PCGExternalDataInterop/Source/PCGExternalDataInterop/Private/Elements/PCGLoadAlembicElement.cpp
new file mode 100644
--- /dev/null
+++ b/Engine/Plugins/PCGInterops/PCGExternalDataInterop/Source/PCGExternalDataInterop/Private/Elements/PCGLoadAlembicElement.cpp
@@ -0,0 +1,160 @@
+// Copyright Epic Games, Inc. All Rights Reserved.
+
+#include "Elements/PCGLoadAlembicElement.h"
+
+#include "Alembic/PCGAlembicInterop.h"
+#include "Misc/Paths.h"
+
+#define LOCTEXT_NAMESPACE "PCGLoadAlembic"
+
+#if WITH_EDITOR
+void UPCGLoadAlembicSettings::PostEditChangeProperty(FPropertyChangedEvent& PropertyChangedEvent)
+{
+	if (PropertyChangedEvent.Property && PropertyChangedEvent.Property->GetFName() == GET_MEMBER_NAME_CHECKED(UPCGLoadAlembicSettings, Setup))
+	{
+		if (Setup == EPCGLoadAlembicStandardSetup::CitySample)
+		{
+			SetupFromStandard(Setup);
+			Setup = EPCGLoadAlembicStandardSetup::None;
+		}
+	}
+
+	Super::PostEditChangeProperty(PropertyChangedEvent);
+}
+
+FText UPCGLoadAlembicSettings::GetDefaultNodeTitle() const
+{
+	return LOCTEXT("NodeTitle", "Load Alembic");
+}
+
+FText UPCGLoadAlembicSettings::GetNodeTooltipText() const
+{
+	return LOCTEXT("NodeTooltip", "Loads data from an Alembic file");
+}
+#endif // WITH_EDITOR
+
+FPCGElementPtr UPCGLoadAlembicSettings::CreateElement() const
+{
+	return MakeShared<FPCGLoadAlembicElement>();
+}
+
+void UPCGLoadAlembicSettings::SetupFromStandard(EPCGLoadAlembicStandardSetup InSetup)
+{
+	if (InSetup == EPCGLoadAlembicStandardSetup::CitySample)
+	{
+		SetupFromStandard(InSetup, ConversionScale, ConversionRotation, bConversionFlipHandedness, AttributeMapping);
+	}
+}
+
+void UPCGLoadAlembicSettings::SetupFromStandard(EPCGLoadAlembicStandardSetup InSetup, FVector& InConversionScale, FVector& InConversionRotation, bool& bInConversionFlipHandedness, TMap<FString, FPCGAttributePropertyInputSelector>& InAttributeMapping)
+{
+	if (InSetup == EPCGLoadAlembicStandardSetup::CitySample)
+	{
+		InConversionScale = FVector(1.0f, 1.0f, 1.0f);
+		InConversionRotation = FVector::ZeroVector;
+		bInConversionFlipHandedness = true;
+
+		InAttributeMapping.Reset();
+
+		FPCGAttributePropertyInputSelector PositionSelector;
+		PositionSelector.Update(FString(TEXT("$Position.xzy")));
+
+		InAttributeMapping.Add(FString(TEXT("position")), MoveTemp(PositionSelector));
+
+		FPCGAttributePropertyInputSelector ScaleSelector;
+		ScaleSelector.Update(FString(TEXT("$Scale.xzy")));
+
+		InAttributeMapping.Add(FString(TEXT("scale")), MoveTemp(ScaleSelector));
+
+		FPCGAttributePropertyInputSelector RotationSelector;
+		RotationSelector.Update(FString(TEXT("$Rotation.xzyw")));
+
+		InAttributeMapping.Add(FString(TEXT("orient")), MoveTemp(RotationSelector));
+	}
+}
+
+FPCGContext* FPCGLoadAlembicElement::CreateContext()
+{
+	return new FPCGLoadAlembicContext();
+}
+
+bool FPCGLoadAlembicElement::PrepareLoad(FPCGExternalDataContext* InContext) const
+{
+	check(InContext);
+	FPCGLoadAlembicContext* Context = static_cast<FPCGLoadAlembicContext*>(InContext);
+
+	check(Context);
+	const UPCGLoadAlembicSettings* Settings = Context->GetInputSettings<UPCGLoadAlembicSettings>();
+	check(Settings);
+
+#if WITH_EDITOR
+	const FString FileName = Settings->AlembicFilePath.FilePath;
+	PCGAlembicInterop::LoadFromAlembicFile(Context, FileName);
+
+	if (!Context->PointDataAccessorsMapping.IsEmpty())
+	{
+		for (const FPCGExternalDataContext::FPointDataAccessorsMapping& DataMapping : Context->PointDataAccessorsMapping)
+		{
+			FPCGTaggedData& OutData = Context->OutputData.TaggedData.Emplace_GetRef();
+			OutData.Data = DataMapping.Data;
+		}
+
+		Context->bDataPrepared = true;
+	}
+#else
+	PCGE_LOG(Error, GraphAndLog, LOCTEXT("NotSupportedInGameMode", "The Load Alembic node is not support in non-editor builds."));
+#endif
+
+	return true;
+}
+
+bool FPCGLoadAlembicElement::ExecuteLoad(FPCGExternalDataContext* InContext) const
+{
+	if (!FPCGExternalDataElement::ExecuteLoad(InContext))
+	{
+		return false;
+	}
+
+	// Finally, apply conversion
+	check(InContext);
+	FPCGLoadAlembicContext* Context = static_cast<FPCGLoadAlembicContext*>(InContext);
+	check(Context);
+	const UPCGLoadAlembicSettings* Settings = Context->GetInputSettings<UPCGLoadAlembicSettings>();
+	check(Settings);
+
+	const FVector& ConversionScale = Settings->ConversionScale;
+	const FVector& ConversionRotation = Settings->ConversionRotation;
+	const bool bFlipRotationW = Settings->bConversionFlipHandedness;
+
+	const FTransform ConversionTransform(FRotator::MakeFromEuler(ConversionRotation), FVector::ZeroVector, ConversionScale);
+	if (!ConversionTransform.Equals(FTransform::Identity) || bFlipRotationW)
+	{
+		for (const FPCGExternalDataContext::FPointDataAccessorsMapping& DataMapping : Context->PointDataAccessorsMapping)
+		{
+			UPCGPointData* PointData = Cast<UPCGPointData>(DataMapping.Data);
+
+			if (!PointData)
+			{
+				continue;
+			}
+
+			TArray<FPCGPoint>& Points = PointData->GetMutablePoints();
+
+			for (FPCGPoint& Point : Points)
+			{
+				Point.Transform = Point.Transform * ConversionTransform;
+
+				if (bFlipRotationW)
+				{
+					FQuat Rotation = Point.Transform.GetRotation();
+					Rotation.W *= -1;
+					Point.Transform.SetRotation(Rotation);
+				}
+			}
+		}
+	}
+
+	return true;
+}
+
+#undef LOCTEXT_NAMESPACE
\ No newline at end of file
diff --git a/Engine/Plugins/PCGInterops/PCGExternalDataInterop/Source/PCGExternalDataInterop/Private/PCGExternalDataInteropModule.cpp b/Engine/Plugins/PCGInterops/PCGExternalDataInterop/Source/PCGExternalDataInterop/Private/PCGExternalDataInteropModule.cpp
new file mode 100644
--- /dev/null
+++ b/Engine/Plugins/PCGInterops/PCGExternalDataInterop/Source/PCGExternalDataInterop/Private/PCGExternalDataInteropModule.cpp
@@ -0,0 +1,16 @@
+// Copyright Epic Games, Inc. All Rights Reserved.
+
+#include "Modules/ModuleManager.h"
+
+class FPCGExternalDataInteropModule final : public IModuleInterface
+{
+public:
+	//~ IModuleInterface implementation
+	virtual bool SupportsDynamicReloading() override
+	{
+		return true;
+	}
+	//~ End IModuleInterface implementation
+};
+
+IMPLEMENT_MODULE(FPCGExternalDataInteropModule, PCGExternalDataInterop);
diff --git a/Engine/Plugins/PCGInterops/PCGExternalDataInterop/Source/PCGExternalDataInterop/Public/Elements/PCGLoadAlembicElement.h b/Engine/Plugins/PCGInterops/PCGExternalDataInterop/Source/PCGExternalDataInterop/Public/Elements/PCGLoadAlembicElement.h
new file mode 100644
--- /dev/null
+++ b/Engine/Plugins/PCGInterops/PCGExternalDataInterop/Source/PCGExternalDataInterop/Public/Elements/PCGLoadAlembicElement.h
@@ -0,0 +1,75 @@
+// Copyright Epic Games, Inc. All Rights Reserved.
+
+#pragma once
+
+#include "Elements/IO/PCGExternalData.h"
+
+#include "PCGLoadAlembicElement.generated.h"
+
+struct FPCGExternalDataContext;
+
+UENUM()
+enum class EPCGLoadAlembicStandardSetup : uint8
+{
+	None = 0,
+	CitySample UMETA(Tooltip="Uses the same setup as in the City Sample demo: right handed Y-up and the orient and scale mapping to the rotation and scale, respectively")
+};
+
+UCLASS(MinimalAPI, BlueprintType, ClassGroup = (Procedural))
+class UPCGLoadAlembicSettings : public UPCGExternalDataSettings
+{
+	GENERATED_BODY()
+
+public:
+	// ~Begin UObject interface
+#if WITH_EDITOR
+	virtual void PostEditChangeProperty(FPropertyChangedEvent& PropertyChangedEvent) override;
+#endif
+	// ~End UObject interface
+
+	// ~Begin UPCGSettings interface
+#if WITH_EDITOR
+	virtual FName GetDefaultNodeName() const override { return FName(TEXT("LoadAlembic")); }
+	virtual FText GetDefaultNodeTitle() const override;
+	virtual FText GetNodeTooltipText() const override;
+#endif
+
+protected:
+	virtual FPCGElementPtr CreateElement() const override;
+	// ~End UPCGSettings interface
+
+public:
+	UPROPERTY(BlueprintReadWrite, EditAnywhere, Category = "Alembic", meta = (FilePathFilter = "Alembic files (*.abc)|*.abc", PCG_Overridable))
+	FFilePath AlembicFilePath;
+
+	// To prevent a dependency on the alembic editor module in this class, we'll keep around only the types we need
+	/** Scale to apply during import. Note that for both Max/Maya presets the value flips the Y axis. */
+	UPROPERTY(BlueprintReadWrite, EditAnywhere, Category = "Alembic", meta = (PCG_Overridable))
+	FVector ConversionScale = FVector(1.0f, -1.0f, 1.0f);
+
+	/** Rotation in Euler angles applied during import. For Max, use (90, 0, 0). */
+	UPROPERTY(BlueprintReadWrite, EditAnywhere, Category = "Alembic", meta = (PCG_Overridable))
+	FVector ConversionRotation = FVector::ZeroVector;
+
+	/** When changing handedness, it is sometimes needed to flip the rotation direction */
+	UPROPERTY(BlueprintReadWrite, EditAnywhere, Category = "Alembic", meta = (PCG_Overridable, Tooltip="Flips rotation direction (W), useful together with swizzling"))
+	bool bConversionFlipHandedness = false;
+
+	UFUNCTION(BlueprintCallable, Category = "PCG|AlembicImport")
+	PCGEXTERNALDATAINTEROP_API void SetupFromStandard(EPCGLoadAlembicStandardSetup InSetup);
+
+	static PCGEXTERNALDATAINTEROP_API void SetupFromStandard(EPCGLoadAlembicStandardSetup InSetup, FVector& InConversionScale, FVector& InConversionRotation, bool& bInConversionFlipHandedness, TMap<FString, FPCGAttributePropertyInputSelector>& InAttributeMapping);
+
+#if WITH_EDITORONLY_DATA
+	UPROPERTY(EditAnywhere, Category = "Alembic", meta = (DisplayName="Setup from standard"))
+	EPCGLoadAlembicStandardSetup Setup = EPCGLoadAlembicStandardSetup::None;
+#endif
+};
+
+class PCGEXTERNALDATAINTEROP_API FPCGLoadAlembicElement : public FPCGExternalDataElement
+{
+protected:
+	virtual FPCGContext* CreateContext() override;
+	virtual bool PrepareLoad(FPCGExternalDataContext* Context) const override;
+	virtual bool ExecuteLoad(FPCGExternalDataContext* Context) const override;
+};
\ No newline at end of file
diff --git a/Engine/Plugins/PCGInterops/PCGExternalDataInterop/Source/PCGExternalDataInteropEditor/PCGExternalDataInteropEditor.Build.cs b/Engine/Plugins/PCGInterops/PCGExternalDataInterop/Source/PCGExternalDataInteropEditor/PCGExternalDataInteropEditor.Build.cs
new file mode 100644
--- /dev/null
+++ b/Engine/Plugins/PCGInterops/PCGExternalDataInterop/Source/PCGExternalDataInteropEditor/PCGExternalDataInteropEditor.Build.cs
@@ -0,0 +1,35 @@
+// Copyright Epic Games, Inc. All Rights Reserved.
+
+namespace UnrealBuildTool.Rules
+{
+	public class PCGExternalDataInteropEditor : ModuleRules
+	{
+		public PCGExternalDataInteropEditor(ReadOnlyTargetRules Target) : base(Target)
+		{
+			PCHUsage = PCHUsageMode.UseExplicitOrSharedPCHs;
+
+			PublicDependencyModuleNames.AddRange(
+				new string[] {
+					"Core",
+					"CoreUObject",
+					"Engine",
+					"Projects",
+					"PCG"
+				});
+
+			if (Target.bBuildEditor == true)
+			{
+				PrivateDependencyModuleNames.AddRange(
+					new string[]
+					{
+						"PCGEditor",
+						"PCGExternalDataInterop",
+						"UnrealEd",
+						"AlembicLib",
+						"AlembicLibrary"
+					}
+				);
+			}
+		}
+	}
+}
diff --git a/Engine/Plugins/PCGInterops/PCGExternalDataInterop/Source/PCGExternalDataInteropEditor/Private/PCGExternalDataInteropEditorModule.cpp b/Engine/Plugins/PCGInterops/PCGExternalDataInterop/Source/PCGExternalDataInteropEditor/Private/PCGExternalDataInteropEditorModule.cpp
new file mode 100644
--- /dev/null
+++ b/Engine/Plugins/PCGInterops/PCGExternalDataInterop/Source/PCGExternalDataInteropEditor/Private/PCGExternalDataInteropEditorModule.cpp
@@ -0,0 +1,16 @@
+// Copyright Epic Games, Inc. All Rights Reserved.
+
+#include "Modules/ModuleManager.h"
+
+class FPCGExternalDataInteropEditorModule final : public IModuleInterface
+{
+public:
+	//~ IModuleInterface implementation
+	virtual bool SupportsDynamicReloading() override
+	{
+		return true;
+	}
+	//~ End IModuleInterface implementation
+};
+
+IMPLEMENT_MODULE(FPCGExternalDataInteropEditorModule, PCGExternalDataInteropEditor);
diff --git a/Engine/Plugins/PCGInterops/PCGExternalDataInterop/Source/PCGExternalDataInteropEditor/Private/PCGLoadAlembic.cpp b/Engine/Plugins/PCGInterops/PCGExternalDataInterop/Source/PCGExternalDataInteropEditor/Private/PCGLoadAlembic.cpp
new file mode 100644
--- /dev/null
+++ b/Engine/Plugins/PCGInterops/PCGExternalDataInterop/Source/PCGExternalDataInteropEditor/Private/PCGLoadAlembic.cpp
@@ -0,0 +1,100 @@
+// Copyright Epic Games, Inc. All Rights Reserved.
+
+#include "PCGLoadAlembic.h"
+#include "PCGAssetExporterUtils.h"
+#include "PCGModule.h"
+
+#include "Elements/PCGLoadAlembicElement.h"
+
+void UPCGAlembicToPCGAssetExporter::SerializeMetadata(FArchive& Ar)
+{
+	FPCGLoadAlembicBPData::StaticStruct()->SerializeItem(Ar, &LoadSettings, nullptr);
+}
+
+bool UPCGAlembicToPCGAssetExporter::ExportAsset(const FString& PackageName, UPCGDataAsset* Asset)
+{
+	check(Asset);
+	UPCGLoadAlembicFunctionLibrary::LoadAlembicFileToPCG(LoadSettings, Asset->Data, Asset);
+	return true;
+}
+
+UPackage* UPCGAlembicToPCGAssetExporter::UpdateAsset(const FAssetData& PCGAsset)
+{
+	UPCGDataAsset* Asset = Cast<UPCGDataAsset>(PCGAsset.GetAsset());
+	if (!Asset)
+	{
+		UE_LOG(LogPCG, Error, TEXT("Asset '%s' isn't a PCG data asset or could not be properly loaded."), *PCGAsset.GetObjectPathString());
+		return nullptr;
+	}
+
+	UPackage* Package = Asset->GetPackage();
+	if (!Package)
+	{
+		UE_LOG(LogPCG, Error, TEXT("Unable to retrieve package from Asset '%s'."), *PCGAsset.GetObjectPathString());
+		return nullptr;
+	}
+
+	UPCGLoadAlembicFunctionLibrary::LoadAlembicFileToPCG(LoadSettings, Asset->Data, Asset);
+
+	return Package;
+}
+
+void UPCGLoadAlembicFunctionLibrary::ExportAlembicFileToPCG(const FPCGLoadAlembicBPData& InSettings, FPCGAssetExporterParameters Parameters)
+{
+	UPCGAlembicToPCGAssetExporter* Exporter = NewObject<UPCGAlembicToPCGAssetExporter>(GetTransientPackage());
+	check(Exporter);
+
+	Exporter->LoadSettings = InSettings;
+	UPCGAssetExporterUtils::CreateAsset(Exporter, Parameters);
+}
+
+void UPCGLoadAlembicFunctionLibrary::LoadAlembicFileToPCG(const FPCGLoadAlembicBPData& InSettings, FPCGDataCollection& Data, UObject* TargetOuter)
+{
+	LoadAlembicFileToPCGInternal(InSettings, Data, TargetOuter);
+}
+
+void UPCGLoadAlembicFunctionLibrary::LoadAlembicFileToPCGInternal(const FPCGLoadAlembicBPData& InSettings, FPCGDataCollection& Data, UObject* TargetOuter)
+{
+	// Simulate execution of the FPCGLoadAlembicElement.
+	UPCGLoadAlembicSettings* Settings = NewObject<UPCGLoadAlembicSettings>();
+	Settings->AlembicFilePath = InSettings.AlembicFilePath;
+	Settings->ConversionScale = InSettings.ConversionSettings.Scale;
+	Settings->ConversionRotation = InSettings.ConversionSettings.Rotation;
+	Settings->bConversionFlipHandedness = InSettings.bConversionFlipHandedness;
+	Settings->AttributeMapping = InSettings.AttributeMapping;
+
+	FPCGLoadAlembicElement Element;
+	FPCGContext* Context = Element.Initialize(FPCGDataCollection(), TWeakObjectPtr<UPCGComponent>(), nullptr);
+	Context->AsyncState.NumAvailableTasks = 1; // TODO : change this
+	FPCGTaggedData& DataSettings = Context->InputData.TaggedData.Emplace_GetRef();
+	DataSettings.Data = Settings;
+
+	while (!Element.Execute(Context))
+	{
+		// Nothing
+	}
+
+	Data = Context->OutputData;
+	delete Context;
+
+	// Reouter the data if needed
+	if (TargetOuter)
+	{
+		for (FPCGTaggedData& TaggedData : Data.TaggedData)
+		{
+			if (TaggedData.Data)
+			{
+				(const_cast<UPCGData*>(TaggedData.Data.Get()))->Rename(nullptr, TargetOuter, REN_DoNotDirty | REN_DontCreateRedirectors);
+			}
+		}
+	}
+}
+
+void UPCGLoadAlembicFunctionLibrary::SetupFromStandard(FPCGLoadAlembicBPData& Data, EPCGLoadAlembicStandardSetup InSetup)
+{
+	if (InSetup != EPCGLoadAlembicStandardSetup::None)
+	{
+		Data.ConversionSettings.Preset = EAbcConversionPreset::Custom;
+		UPCGLoadAlembicSettings::SetupFromStandard(InSetup, Data.ConversionSettings.Scale, Data.ConversionSettings.Rotation, Data.bConversionFlipHandedness, Data.AttributeMapping);
+	}	
+}
\ No newline at end of file
diff --git a/Engine/Plugins/PCGInterops/PCGExternalDataInterop/Source/PCGExternalDataInteropEditor/Public/PCGLoadAlembic.h b/Engine/Plugins/PCGInterops/PCGExternalDataInterop/Source/PCGExternalDataInteropEditor/Public/PCGLoadAlembic.h
new file mode 100644
--- /dev/null
+++ b/Engine/Plugins/PCGInterops/PCGExternalDataInterop/Source/PCGExternalDataInteropEditor/Public/PCGLoadAlembic.h
@@ -0,0 +1,67 @@
+// Copyright Epic Games, Inc. All Rights Reserved.
+
+#pragma once
+
+#include "PCGAssetExporter.h"
+#include "Elements/PCGLoadAlembicElement.h"
+#include "Metadata/PCGAttributePropertySelector.h"
+
+#include "AbcImportSettings.h"
+
+#include "PCGLoadAlembic.generated.h"
+
+USTRUCT(BlueprintType)
+struct PCGEXTERNALDATAINTEROPEDITOR_API FPCGLoadAlembicBPData
+{
+	GENERATED_BODY()
+
+	UPROPERTY(EditAnywhere, Category = "Alembic", meta = (FilePathFilter = "Alembic files (*.abc)|*.abc", PCG_Overridable))
+	FFilePath AlembicFilePath;
+
+	/** Conversion settings that will be applied on the transform only */
+	UPROPERTY(BlueprintReadWrite, EditAnywhere, Category = "Settings")
+	FAbcConversionSettings ConversionSettings;
+
+	UPROPERTY(BlueprintReadWrite, EditAnywhere, Category = "Settings", meta = (Tooltip = "Flips rotation direction (W), useful together with swizzling"))
+	bool bConversionFlipHandedness = false;
+
+	UPROPERTY(BlueprintReadWrite, EditAnywhere, Category = "Settings")
+	TMap<FString, FPCGAttributePropertyInputSelector> AttributeMapping;
+};
+
+UCLASS()
+class UPCGAlembicToPCGAssetExporter : public UPCGAssetExporter
+{
+	GENERATED_BODY()
+
+public:
+	//~Begin UPCGAssetExporter interface
+	virtual bool ExportAsset(const FString& PackageName, UPCGDataAsset* Asset);
+	virtual UPackage* UpdateAsset(const FAssetData& PCGAsset) override;
+
+protected:
+	virtual void SerializeMetadata(FArchive& Ar) override;
+	//~End UPCGAssetExporter interface
+
+public:
+	FPCGLoadAlembicBPData LoadSettings;
+};
+
+UCLASS()
+class PCGEXTERNALDATAINTEROPEDITOR_API UPCGLoadAlembicFunctionLibrary : public UBlueprintFunctionLibrary
+{
+	GENERATED_BODY()
+
+public:
+	UFUNCTION(BlueprintCallable, Category = "PCG|IO")
+	static void ExportAlembicFileToPCG(const FPCGLoadAlembicBPData& Settings, FPCGAssetExporterParameters Parameters = FPCGAssetExporterParameters());
+
+	UFUNCTION(BlueprintCallable, Category = "PCG|IO", meta = (DeprecatedFunction, DeprecationMessage="The LoadAlembicFileToPCG function has been replaced by ExportAlembicFileToPCG"))
+	static void LoadAlembicFileToPCG(const FPCGLoadAlembicBPData& Settings, FPCGDataCollection& Data, UObject* TargetOuter);
+
+	UFUNCTION(BlueprintCallable, Category = "PCG|IO", meta = (ScriptMethod))
+	static void SetupFromStandard(UPARAM(ref) FPCGLoadAlembicBPData& Data, EPCGLoadAlembicStandardSetup InSetup);
+
+protected:
+	static void LoadAlembicFileToPCGInternal(const FPCGLoadAlembicBPData& Settings, FPCGDataCollection& Data, UObject* TargetOuter);
+};
\ No newline at end of file
diff --git a/Engine/Plugins/PCGInterops/PCGGeometryScriptInterop/PCGGeometryScriptInterop.uplugin b/Engine/Plugins/PCGInterops/PCGGeometryScriptInterop/PCGGeometryScriptInterop.uplugin
new file mode 100644
--- /dev/null
+++ b/Engine/Plugins/PCGInterops/PCGGeometryScriptInterop/PCGGeometryScriptInterop.uplugin
@@ -0,0 +1,34 @@
+{
+	"FileVersion": 3,
+	"Version": 1,
+	"VersionName": "0.2",
+	"FriendlyName": "Procedural Content Generation Framework (PCG) Geometry Script Interop",
+	"Description": "Extra plugin for Procedural Content Generation Framework interacting with Geometry Scripts.",
+	"Category": "Editor",
+	"CreatedBy": "Epic Games, Inc.",
+	"CreatedByURL": "https://epicgames.com",
+	"DocsURL": "https://docs.unrealengine.com/latest/en-US/procedural-content-generation--framework-in-unreal-engine/",
+	"MarketplaceURL": "",
+	"SupportURL": "",
+	"EnabledByDefault": false,
+	"CanContainContent": true,
+	"IsBetaVersion": true,
+	"Installed": false,
+	"Modules": [
+		{
+			"Name": "PCGGeometryScriptInterop",
+			"Type": "Runtime",
+			"LoadingPhase": "Default"
+		}
+	],
+	"Plugins": [
+		{
+			"Name": "PCG",
+			"Enabled": true
+		},
+		{
+			"Name": "GeometryScripting",
+			"Enabled": true
+		}
+	]
+}
diff --git a/Engine/Plugins/PCGInterops/PCGGeometryScriptInterop/Source/PCGGeometryScriptInterop/PCGGeometryScriptInterop.Build.cs b/Engine/Plugins/PCGInterops/PCGGeometryScriptInterop/Source/PCGGeometryScriptInterop/PCGGeometryScriptInterop.Build.cs
new file mode 100644
--- /dev/null
+++ b/Engine/Plugins/PCGInterops/PCGGeometryScriptInterop/Source/PCGGeometryScriptInterop/PCGGeometryScriptInterop.Build.cs
@@ -0,0 +1,43 @@
+// Copyright Epic Games, Inc. All Rights Reserved.
+
+namespace UnrealBuildTool.Rules
+{
+	public class PCGGeometryScriptInterop : ModuleRules
+	{
+		public PCGGeometryScriptInterop(ReadOnlyTargetRules Target) : base(Target)
+		{
+			PCHUsage = PCHUsageMode.UseExplicitOrSharedPCHs;
+
+			PublicDependencyModuleNames.AddRange(
+				new string[] {
+					"Core",
+					"CoreUObject",
+					"Engine",
+					"GeometryScriptingCore",
+					"Projects",
+					"RenderCore",
+					"RHI",
+					"PCG"
+				});
+
+			PrivateDependencyModuleNames.AddRange(
+				new string[]
+				{
+					"GeometryCore",
+					"GeometryFramework",
+					"ModelingComponents"
+				}
+			);
+
+			if (Target.bBuildEditor == true)
+			{
+				PrivateDependencyModuleNames.AddRange(
+					new string[]
+					{
+						"UnrealEd",
+					}
+				);
+			}
+		}
+	}
+}
diff --git a/Engine/Plugins/PCGInterops/PCGGeometryScriptInterop/Source/PCGGeometryScriptInterop/Private/Elements/PCGMeshSampler.cpp b/Engine/Plugins/PCGInterops/PCGGeometryScriptInterop/Source/PCGGeometryScriptInterop/Private/Elements/PCGMeshSampler.cpp
new file mode 100644
--- /dev/null
+++ b/Engine/Plugins/PCGInterops/PCGGeometryScriptInterop/Source/PCGGeometryScriptInterop/Private/Elements/PCGMeshSampler.cpp
@@ -0,0 +1,703 @@
+// Copyright Epic Games, Inc. All Rights Reserved.
+
+#include "Elements/PCGMeshSampler.h"
+
+#include "Helpers/PCGGeometryHelpers.h"
+
+#include "PCGComponent.h"
+#include "PCGModule.h"
+#include "PCGPin.h"
+#include "Data/PCGPointData.h"
+#include "Helpers/PCGAsync.h"
+#include "Helpers/PCGBlueprintHelpers.h"
+#include "Helpers/PCGHelpers.h"
+#include "Metadata/PCGMetadata.h"
+#include "Metadata/PCGMetadataAttributeTpl.h"
+
+#include "UDynamicMesh.h"
+#include "ConversionUtils/SceneComponentToDynamicMesh.h"
+#include "Engine/AssetManager.h"
+#include "Engine/StaticMesh.h"
+#include "GameFramework/Actor.h"
+#include "GeometryScript/MeshAssetFunctions.h"
+#include "GeometryScript/MeshMaterialFunctions.h"
+#include "GeometryScript/MeshNormalsFunctions.h"
+#include "GeometryScript/MeshQueryFunctions.h"
+#include "GeometryScript/MeshRepairFunctions.h"
+#include "GeometryScript/MeshVertexColorFunctions.h"
+#include "GeometryScript/MeshVoxelFunctions.h"
+#include "GeometryScript/SceneUtilityFunctions.h"
+#include "Sampling/MeshSurfacePointSampling.h"
+
+#define LOCTEXT_NAMESPACE "PCGMeshSampler"
+
+namespace PCGMeshSampler
+{
+	bool SampleOnePointPerVertex(const UPCGMeshSamplerSettings* Settings, FPCGMeshSamplerContext* Context, FPCGMeshSamplerContext::SetPointDensityFunc SetPointDensityPtr)
+	{
+		check(Settings && Context && SetPointDensityPtr);
+
+		TRACE_CPUPROFILER_EVENT_SCOPE(PCGMeshSampler::SampleOnePointPerVertex);
+
+		auto IterationBody = [Settings, Context, SetPointDensityPtr](int32 ReadIndex, int32 WriteIndex) -> bool
+		{
+			int32 DataIndex = 0;
+			while (!Context->DynamicMeshes[DataIndex] || ReadIndex >= Context->StartingIndices[DataIndex + 1])
+			{
+				DataIndex++;
+			}
+
+			const TArray<FVector>& Positions = *Context->Positions[DataIndex].List.Get();
+			const TArray<FLinearColor>& Colors = *Context->Colors[DataIndex].List.Get();
+			const TArray<FVector>& Normals = *Context->Normals[DataIndex].List.Get();
+
+			const int32 CurrentIndex = ReadIndex - Context->StartingIndices[DataIndex];
+
+			const FVector& Position = Positions[CurrentIndex];
+			const FLinearColor& Color = Colors[CurrentIndex];
+			const FVector& Normal = Normals[CurrentIndex];
+
+			FPCGPoint OutPoint{};
+			OutPoint.Transform = FTransform{ FRotationMatrix::MakeFromZ(Normal).Rotator(), Position, FVector::OneVector };
+			OutPoint.Color = Color;
+			OutPoint.Steepness = Settings->PointSteepness;
+
+			SetPointDensityPtr(Color, OutPoint);
+
+			UPCGBlueprintHelpers::SetSeedFromPosition(OutPoint);
+
+			Context->OutPointData[DataIndex]->GetMutablePoints()[CurrentIndex] = OutPoint;
+
+			return true;
+		};
+
+		// Small verification to catch any problem that could arise in the iteration body
+		if (Context->DynamicMeshes.IsEmpty() || !ensure(Context->DynamicMeshes.Num() + 1 == Context->StartingIndices.Num()))
+		{
+			return true;
+		}
+
+		return FPCGAsync::AsyncProcessingOneToOneEx(&Context->AsyncState, Context->StartingIndices.Last(), []() {}, IterationBody, /*bEnableTimeSlicing=*/true);
+	}
+
+	bool SampleOnePointPerTriangle(const UPCGMeshSamplerSettings* Settings, FPCGMeshSamplerContext* Context, FPCGMeshSamplerContext::SetPointDensityFunc SetPointDensityPtr)
+	{
+		check(Settings && Context && SetPointDensityPtr);
+
+		TRACE_CPUPROFILER_EVENT_SCOPE(FPCGMeshSamplerElement::Execute::OnePointPerTriangle);
+
+		auto IterationBody = [Context, Settings, SetPointDensityPtr](int32 ReadIndex, int32 WriteIndex) -> bool
+		{
+			int32 DataIndex = 0;
+			while (!Context->DynamicMeshes[DataIndex] || ReadIndex >= Context->StartingIndices[DataIndex + 1])
+			{
+				DataIndex++;
+			}
+
+			const TArray<int32>& TriangleIds = *Context->TriangleIds[DataIndex].List.Get();
+			const int32 CurrentIndex = ReadIndex - Context->StartingIndices[DataIndex];
+
+			const int32 TriangleId = TriangleIds[CurrentIndex];
+
+			FVector Vertex1, Vertex2, Vertex3;
+			bool bIsValidTriangle;
+
+			UGeometryScriptLibrary_MeshQueryFunctions::GetTrianglePositions(Context->DynamicMeshes[DataIndex], TriangleId, bIsValidTriangle, Vertex1, Vertex2, Vertex3);
+			const FVector Normal = UGeometryScriptLibrary_MeshQueryFunctions::GetTriangleFaceNormal(Context->DynamicMeshes[DataIndex], TriangleId, bIsValidTriangle);
+			const FVector Position = (Vertex1 + Vertex2 + Vertex3) / 3.0;
+
+			FPCGPoint OutPoint{};
+			OutPoint.Transform = FTransform{ FRotationMatrix::MakeFromZ(Normal).Rotator(), Position, FVector::OneVector };
+			OutPoint.Steepness = Settings->PointSteepness;
+
+			FVector Dummy1, Dummy2, Dummy3;
+			FVector BarycentricCoord;
+			bool bIsValid = false;
+			UGeometryScriptLibrary_MeshQueryFunctions::ComputeTriangleBarycentricCoords(Context->DynamicMeshes[DataIndex], TriangleId, bIsValid, Position, Dummy1, Dummy2, Dummy3, BarycentricCoord);
+
+			Context->SetPointColorAndDensity(SetPointDensityPtr, TriangleId, BarycentricCoord, OutPoint, DataIndex);
+			Context->SetAttributeValues(Settings->UVChannel, TriangleId, BarycentricCoord, OutPoint, DataIndex);
+
+			UPCGBlueprintHelpers::SetSeedFromPosition(OutPoint);
+
+			Context->OutPointData[DataIndex]->GetMutablePoints()[CurrentIndex] = OutPoint;
+
+			return true;
+		};
+
+		// Small verification to catch any problem that could arise in the iteration body
+		if (Context->DynamicMeshes.IsEmpty() || !ensure(Context->DynamicMeshes.Num() + 1 == Context->StartingIndices.Num()))
+		{
+			return true;
+		}
+
+		return FPCGAsync::AsyncProcessingOneToOneEx(&Context->AsyncState, Context->StartingIndices.Last(), []() {}, IterationBody, /*bEnableTimeSlicing=*/true);
+	}
+
+	bool PoissonSampling(const UPCGMeshSamplerSettings* Settings, FPCGMeshSamplerContext* Context, FPCGMeshSamplerContext::SetPointDensityFunc SetPointDensityPtr)
+	{
+		check(Settings && Context && SetPointDensityPtr);
+
+		// For Poisson sampling, we are calling a "all-in-one" function, where we don't have control for timeslicing.
+		// Since Poisson sampling can be expensive (depending on the radius used), we will do the sampling in a future, put this task to sleep,
+		// and wait for the sampling to wake us up.
+		if (Context->SamplingFutures.IsEmpty())
+		{
+			// Put the task asleep
+			Context->bIsPaused = true;
+			Context->SamplingProgess = MakeUnique<FProgressCancel>();
+			Context->SamplingProgess->CancelF = [Context]() -> bool { return Context->StopSampling; };
+
+			for (int32 DataIndex = 0; DataIndex < Context->DynamicMeshes.Num(); ++DataIndex)
+			{
+				if (!Context->DynamicMeshes[DataIndex])
+				{
+					continue;
+				}
+
+				auto SamplingFuture = [Settings, Context, Seed = Context->GetSeed(), SetPointDensityPtr, DataIndex]() -> bool
+				{
+					TRACE_CPUPROFILER_EVENT_SCOPE(FPCGMeshSamplerElement::Execute::PoissonSampling);
+
+					UE::Geometry::FMeshSurfacePointSampling PointSampling;
+					PointSampling.SampleRadius = Settings->SamplingOptions.SamplingRadius;
+					PointSampling.MaxSamples = Settings->SamplingOptions.MaxNumSamples;
+					PointSampling.RandomSeed = PCGHelpers::ComputeSeed(Seed, Settings->SamplingOptions.RandomSeed);
+					PointSampling.SubSampleDensity = Settings->SamplingOptions.SubSampleDensity;
+
+					if (Settings->NonUniformSamplingOptions.MaxSamplingRadius > PointSampling.SampleRadius)
+					{
+						PointSampling.MaxSampleRadius = Settings->NonUniformSamplingOptions.MaxSamplingRadius;
+						PointSampling.SizeDistribution = UE::Geometry::FMeshSurfacePointSampling::ESizeDistribution(static_cast<int>(Settings->NonUniformSamplingOptions.SizeDistribution));
+						PointSampling.SizeDistributionPower = FMath::Clamp(Settings->NonUniformSamplingOptions.SizeDistributionPower, 1.0, 10.0);
+					}
+
+					PointSampling.bComputeBarycentrics = true;
+
+					PointSampling.ComputePoissonSampling(Context->DynamicMeshes[DataIndex]->GetMeshRef(), Context->SamplingProgess.Get());
+
+					if (Context->StopSampling)
+					{
+						return true;
+					}
+
+					TArray<FPCGPoint>& Points = Context->OutPointData[DataIndex]->GetMutablePoints();
+					Points.Reserve(PointSampling.Samples.Num());
+
+					int Count = 0;
+					for (int32 i = 0; i < PointSampling.Samples.Num(); ++i)
+					{
+						UE::Geometry::FFrame3d& Sample = PointSampling.Samples[i];
+						// Avoid to check too many times
+						constexpr int CancelledCheckNum = 25;
+						if (++Count == CancelledCheckNum)
+						{
+							Count = 0;
+							if (Context->StopSampling)
+							{
+								return true;
+							}
+						}
+
+						const int32 TriangleId = PointSampling.TriangleIDs[i];
+						const FVector BarycentricCoords = PointSampling.BarycentricCoords[i];
+
+						FPCGPoint& OutPoint = Points.Emplace_GetRef();
+						OutPoint.Transform = Sample.ToTransform();
+						OutPoint.Steepness = Settings->PointSteepness;
+
+						Context->SetPointColorAndDensity(SetPointDensityPtr, TriangleId, BarycentricCoords, OutPoint, DataIndex);
+						Context->SetAttributeValues(Settings->UVChannel, TriangleId, BarycentricCoords, OutPoint, DataIndex);
+
+						UPCGBlueprintHelpers::SetSeedFromPosition(OutPoint);
+					}
+
+					// Unpause the task
+					Context->bIsPaused = false;
+					return true;
+				};
+
+				Context->SamplingFutures.Add(Async(EAsyncExecution::ThreadPool, std::move(SamplingFuture)));
+			}
+		}
+
+		const bool bIsDone = Algo::AllOf(Context->SamplingFutures, [](const TFuture<bool>& Future) { return Future.IsReady();});
+		if (bIsDone)
+		{
+			for (TFuture<bool>& Future : Context->SamplingFutures)
+			{
+				Future.Reset();
+			}
+		}
+
+		return bIsDone;
+	}
+}
+
+void FPCGMeshSamplerContext::SetAttributeValues(int32 UVChannel, int32 TriangleId, const FVector& BarycentricCoord, FPCGPoint& OutPoint, int32 DataIndex)
+{
+	if (UVAttributes.IsValidIndex(DataIndex) && UVAttributes[DataIndex])
+	{
+		bool bHasValidUVs = false;
+		FVector2D InterpolatedUV = FVector2D::ZeroVector;
+		UGeometryScriptLibrary_MeshQueryFunctions::GetInterpolatedTriangleUV(DynamicMeshes[DataIndex], /*UVSetIndex=*/ UVChannel, TriangleId, BarycentricCoord, bHasValidUVs, InterpolatedUV);
+		if (bHasValidUVs)
+		{
+			check(OutPointData[DataIndex]);
+			OutPointData[DataIndex]->Metadata->InitializeOnSet(OutPoint.MetadataEntry);
+			UVAttributes[DataIndex]->SetValue(OutPoint.MetadataEntry, InterpolatedUV);
+		}
+	}
+
+	if (MaterialIdAttributes.IsValidIndex(DataIndex) && MaterialIdAttributes[DataIndex])
+	{
+		bool bIsValidTriangle = false;
+		const int32 MaterialId = UGeometryScriptLibrary_MeshMaterialFunctions::GetTriangleMaterialID(DynamicMeshes[DataIndex], TriangleId, bIsValidTriangle);
+
+		if (bIsValidTriangle)
+		{
+			check(OutPointData[DataIndex]);
+			OutPointData[DataIndex]->Metadata->InitializeOnSet(OutPoint.MetadataEntry);
+			MaterialIdAttributes[DataIndex]->SetValue(OutPoint.MetadataEntry, MaterialId);
+
+			if (MaterialAttributes.IsValidIndex(DataIndex) && MaterialAttributes[DataIndex] && AssetMaterialList.IsValidIndex(DataIndex) && AssetMaterialList[DataIndex].IsValidIndex(MaterialId) && AssetMaterialList[DataIndex][MaterialId])
+			{
+				MaterialAttributes[DataIndex]->SetValue(OutPoint.MetadataEntry, AssetMaterialList[DataIndex][MaterialId]);
+			}
+		}
+	}
+
+	if (TriangleIdAttributes.IsValidIndex(DataIndex) && TriangleIdAttributes[DataIndex])
+	{
+		check(OutPointData[DataIndex]);
+		OutPointData[DataIndex]->Metadata->InitializeOnSet(OutPoint.MetadataEntry);
+		TriangleIdAttributes[DataIndex]->SetValue(OutPoint.MetadataEntry, TriangleId);
+	}
+}
+
+void FPCGMeshSamplerContext::SetPointColorAndDensity(SetPointDensityFunc SetPointDensityFuncPtr, int32 TriangleId, const FVector& BarycentricCoord, FPCGPoint& OutPoint, int32 DataIndex)
+{
+	FLinearColor Color;
+	bool bValidVertexColor = false;
+	UGeometryScriptLibrary_MeshQueryFunctions::GetInterpolatedTriangleVertexColor(DynamicMeshes[DataIndex], TriangleId, BarycentricCoord, FLinearColor::White, bValidVertexColor, Color);
+	if (bValidVertexColor)
+	{
+		OutPoint.Color = Color;
+		SetPointDensityFuncPtr(Color, OutPoint);
+	}
+	else
+	{
+		OutPoint.Density = 1.0f;
+	}
+}
+
+UPCGMeshSamplerSettings::UPCGMeshSamplerSettings()
+	: UPCGSettings()
+{
+	bUseSeed = true;
+}
+
+void UPCGMeshSamplerSettings::PostLoad()
+{
+	Super::PostLoad();
+
+#if WITH_EDITOR
+	if (!StaticMeshPath_DEPRECATED.IsNull())
+	{
+		StaticMesh = StaticMeshPath_DEPRECATED;
+		StaticMeshPath_DEPRECATED.Reset();
+	}
+
+	if (bUseRedAsDensity_DEPRECATED)
+	{
+		// It was only available for one point per vertex before. Keep that.
+		bUseColorChannelAsDensity = (SamplingMethod == EPCGMeshSamplingMethod::OnePointPerVertex);
+		ColorChannelAsDensity = EPCGColorChannel::Red;
+		bUseRedAsDensity_DEPRECATED = false;
+	}
+#endif
+}
+
+TArray<FPCGPinProperties> UPCGMeshSamplerSettings::InputPinProperties() const
+{
+	TArray<FPCGPinProperties> Properties;
+	if (bExtractMeshFromInput)
+	{
+		Properties.Emplace_GetRef(PCGPinConstants::DefaultInputLabel, EPCGDataType::Any).SetRequiredPin();
+	}
+
+	return Properties;
+}
+
+TArray<FPCGPinProperties> UPCGMeshSamplerSettings::OutputPinProperties() const
+{
+	TArray<FPCGPinProperties> Properties;
+	Properties.Emplace(PCGPinConstants::DefaultOutputLabel, EPCGDataType::Point, /*bInAllowMultipleConnections=*/bExtractMeshFromInput, /*bAllowMultipleData=*/bExtractMeshFromInput);
+	return Properties;
+}
+
+FPCGElementPtr UPCGMeshSamplerSettings::CreateElement() const
+{
+	return MakeShared<FPCGMeshSamplerElement>();
+}
+
+#if WITH_EDITOR
+FText UPCGMeshSamplerSettings::GetNodeTooltipText() const
+{
+	return LOCTEXT("MeshSamplerNodeTooltip", "Sample points on a static mesh.");
+}
+#endif // WITH_EDITOR
+
+FPCGMeshSamplerContext::~FPCGMeshSamplerContext()
+{
+	// The context can be destroyed if the task is canceled. In that case, we need to notify the futures that they should stop,
+	// and wait for them to finish (as the futures use some data stored in the context that will go dangling if the context is destroyed).
+	StopSampling = true;
+	for (TFuture<bool>& SamplingFuture : SamplingFutures)
+	{
+		if (SamplingFuture.IsValid() && !SamplingFuture.IsReady())
+		{
+			SamplingFuture.Wait();
+		}
+	}
+}
+
+void FPCGMeshSamplerContext::AddExtraStructReferencedObjects(FReferenceCollector& Collector)
+{
+	for (TObjectPtr<UDynamicMesh>& DynamicMesh : DynamicMeshes)
+	{
+		if (DynamicMesh)
+		{
+			Collector.AddReferencedObject(DynamicMesh);
+		}
+	}
+}
+
+FPCGContext* FPCGMeshSamplerElement::CreateContext()
+{
+	return new FPCGMeshSamplerContext();
+}
+
+bool FPCGMeshSamplerElement::PrepareDataInternal(FPCGContext* InContext) const
+{
+	TRACE_CPUPROFILER_EVENT_SCOPE(FPCGMeshSamplerElement::PrepareData);
+	FPCGMeshSamplerContext* Context = static_cast<FPCGMeshSamplerContext*>(InContext);
+
+	check(Context);
+
+	if (!Context->SourceComponent.IsValid())
+	{
+		return true;
+	}
+
+	const UPCGMeshSamplerSettings* Settings = Context->GetInputSettings<UPCGMeshSamplerSettings>();
+	check(Settings);
+
+	const TSoftObjectPtr<UStaticMesh> StaticMeshPtr = Settings->StaticMesh;
+
+	// 1. Request load for meshes/inputs. Return false if we need to wait, otherwise continue.
+	if (!Context->WasLoadRequested())
+	{
+		const bool bNeedToWait = !Context->InitializeAndRequestLoad(PCGPinConstants::DefaultInputLabel, Settings->InputSource, { StaticMeshPtr.ToSoftObjectPath() }, /*bPersistAllData=*/false, /*bSilenceErrorOnEmptyObjectPath=*/true, Settings->bSynchronousLoad);
+
+		if (bNeedToWait)
+		{
+			return false;
+		}
+	}
+
+	UGeometryScriptDebug* Debug = FPCGContext::NewObject_AnyThread<UGeometryScriptDebug>(Context);
+
+	// Make sure the LOD is valid, and make the conversion.
+	UE::Conversion::EMeshLODType RequestedLODType;
+
+	switch (Settings->RequestedLODType)
+	{
+	case EGeometryScriptLODType::MaxAvailable:
+		RequestedLODType = UE::Conversion::EMeshLODType::MaxAvailable;
+		break;
+	case EGeometryScriptLODType::HiResSourceModel:
+		RequestedLODType = UE::Conversion::EMeshLODType::HiResSourceModel;
+		break;
+	case EGeometryScriptLODType::SourceModel:
+		RequestedLODType = UE::Conversion::EMeshLODType::SourceModel;
+		break;
+	case EGeometryScriptLODType::RenderData:
+		RequestedLODType = UE::Conversion::EMeshLODType::RenderData;
+		break;
+	default:
+	{
+		PCGLog::LogErrorOnGraph(LOCTEXT("InvalidLOD", "Requested LOD is invalid."), Context);
+		return true;
+	}
+	}
+
+	for (const auto& [Path, DummyIndex, DummyIndex2] : Context->PathsToObjectsAndDataIndex)
+	{
+		UObject* Object = Path.ResolveObject();
+		if (!Object)
+		{
+			PCGE_LOG(Error, GraphAndLog, FText::Format(LOCTEXT("ObjectDoesNotExist", "Provided object does not exist or could not be loaded: '{0}'"), FText::FromString(Path.ToString())));
+			return true;
+		}
+
+		EGeometryScriptOutcomePins Outcome = EGeometryScriptOutcomePins::Success;
+		Debug->Messages.Reset();
+		USceneComponent* SceneComponent = Cast<USceneComponent>(Object);
+
+		// @todo_pcg: Better support of multiple scene component on a given actor.
+		if (AActor* Actor = Cast<AActor>(Object))
+		{
+			SceneComponent = Actor->GetRootComponent();
+		}
+
+		if (SceneComponent)
+		{
+			// Adaptation of UGeometryScriptLibrary_SceneUtilityFunctions::CopyMeshFromComponent, since we don't have access to the Material list with this one.
+			Context->DynamicMeshes.Add(FPCGContext::NewObject_AnyThread<UDynamicMesh>(Context));
+			FTransform Transform;
+			FText ErrorMessage;
+			UE::Conversion::FToMeshOptions Options{};
+			Options.bWantInstanceColors = true;
+			Options.LODType = RequestedLODType;
+			Options.LODIndex = Settings->RequestedLODIndex;
+
+			UE::Geometry::FDynamicMesh3 TempDynMesh{};
+
+			TArray<UMaterialInterface*>* ComponentMaterialListPtr = Settings->bOutputMaterialInfo ? &Context->ComponentMaterialList.Emplace_GetRef() : nullptr;
+			TArray<UMaterialInterface*>* AssetMaterialListPtr = Settings->bOutputMaterialInfo ? &Context->AssetMaterialList.Emplace_GetRef() : nullptr;
+
+			const bool bSuccess = UE::Conversion::SceneComponentToDynamicMesh(SceneComponent, Options, /*bTransformToWorld=*/false, TempDynMesh, Transform, ErrorMessage, ComponentMaterialListPtr, AssetMaterialListPtr);
+			if (!bSuccess)
+			{
+				Outcome = EGeometryScriptOutcomePins::Failure;
+				FGeometryScriptDebugMessage DebugMessage;
+				DebugMessage.Message = std::move(ErrorMessage);
+				Debug->Messages.Emplace(std::move(DebugMessage));
+			}
+			else
+			{
+				Context->DynamicMeshes.Last()->SetMesh(std::move(TempDynMesh));
+			}
+		}
+		else if (UStaticMesh* StaticMesh = Cast<UStaticMesh>(Object))
+		{
+			const FGeometryScriptMeshReadLOD MeshReadLOD{ Settings->RequestedLODType, Settings->RequestedLODIndex };
+
+			Context->DynamicMeshes.Add(FPCGContext::NewObject_AnyThread<UDynamicMesh>(Context));
+			FGeometryScriptCopyMeshFromAssetOptions Options{};
+
+			UGeometryScriptLibrary_StaticMeshFunctions::CopyMeshFromStaticMesh(StaticMesh, Context->DynamicMeshes.Last(), Options, MeshReadLOD, Outcome, Debug);
+			if (Outcome == EGeometryScriptOutcomePins::Success && Settings->bOutputMaterialInfo)
+			{
+				UGeometryScriptLibrary_StaticMeshFunctions::GetMaterialListFromStaticMesh(StaticMesh, Context->AssetMaterialList.Emplace_GetRef());
+			}
+		}
+		else
+		{
+			PCGE_LOG(Error, GraphAndLog, FText::Format(LOCTEXT("ObjectNotRightType", "Provided object '{0}' is not a supported type. Only supports StaticMesh/Actor/SceneComponent."), FText::FromString(Path.ToString())));
+			continue;
+		}
+
+		if (Outcome != EGeometryScriptOutcomePins::Success)
+		{
+			PCGLog::LogErrorOnGraph(FText::Format(LOCTEXT("ObjectToDynamicMeshFailed", "Object to Dynamic mesh failed for object {0}."), FText::FromString(Path.ToString())));
+			PCGGeometryHelpers::GeometryScriptDebugToPCGLog(Context, Debug);
+			Context->DynamicMeshes.Last()->MarkAsGarbage();
+			Context->DynamicMeshes.RemoveAtSwap(Context->DynamicMeshes.Num() - 1);
+		}
+	}
+
+	Debug->MarkAsGarbage();
+
+	// Reserve arrays. Add one more entry for the starting indices to have the total number of items to process.
+	switch (Settings->SamplingMethod)
+	{
+	case EPCGMeshSamplingMethod::OnePointPerVertex:
+		Context->Positions.SetNum(Context->DynamicMeshes.Num());
+		Context->Colors.SetNum(Context->DynamicMeshes.Num());
+		Context->Normals.SetNum(Context->DynamicMeshes.Num());
+		Context->StartingIndices.SetNumZeroed(Context->DynamicMeshes.Num() + 1);
+		break;
+	case EPCGMeshSamplingMethod::OnePointPerTriangle:
+		Context->TriangleIds.SetNum(Context->DynamicMeshes.Num());
+		Context->StartingIndices.SetNumZeroed(Context->DynamicMeshes.Num() + 1);
+		break;
+	default:
+		break;
+	}
+	
+	for (int32 i = 0; i < Context->DynamicMeshes.Num(); ++i)
+	{
+		UDynamicMesh* DynamicMesh = Context->DynamicMeshes[i];
+		if (!DynamicMesh)
+		{
+			continue;
+		}
+
+		if (Settings->bVoxelize)
+		{
+			FGeometryScriptSolidifyOptions SolidifyOptions{};
+			SolidifyOptions.GridParameters.GridCellSize = Settings->VoxelSize;
+			SolidifyOptions.GridParameters.SizeMethod = EGeometryScriptGridSizingMethod::GridCellSize;
+
+			UGeometryScriptLibrary_MeshVoxelFunctions::ApplyMeshSolidify(DynamicMesh, SolidifyOptions);
+
+			if (Settings->bRemoveHiddenTriangles)
+			{
+				FGeometryScriptRemoveHiddenTrianglesOptions RemoveTriangleOptions{};
+
+				UGeometryScriptLibrary_MeshRepairFunctions::RemoveHiddenTriangles(DynamicMesh, RemoveTriangleOptions);
+			}
+		}
+
+		TArray<FPCGTaggedData>& Outputs = InContext->OutputData.TaggedData;
+		UPCGPointData* CurrentOutPointData = FPCGContext::NewObject_AnyThread<UPCGPointData>(Context);
+		Context->OutPointData.Emplace(CurrentOutPointData);
+
+		int32 NumIterations = -1;
+
+		switch (Settings->SamplingMethod)
+		{
+		case EPCGMeshSamplingMethod::OnePointPerVertex:
+		{
+			bool Dummy, Dummy2;
+			UGeometryScriptLibrary_MeshQueryFunctions::GetAllVertexPositions(DynamicMesh, Context->Positions[i], /*bSkipGaps=*/false, /*bHasVertexIDGaps=*/Dummy);
+			UGeometryScriptLibrary_MeshVertexColorFunctions::GetMeshPerVertexColors(DynamicMesh, Context->Colors[i], /*bIsValidColorSet=*/Dummy, /*bHasVertexIDGaps=*/Dummy2);
+			UGeometryScriptLibrary_MeshNormalsFunctions::GetMeshPerVertexNormals(DynamicMesh, Context->Normals[i], /*bIsValidNormalSet=*/Dummy, /*bHasVertexIDGaps=*/Dummy2);
+			NumIterations = Context->Positions[i].List->Num();
+			break;
+		}
+		case EPCGMeshSamplingMethod::OnePointPerTriangle:
+		{
+			bool Dummy;
+			UGeometryScriptLibrary_MeshQueryFunctions::GetAllTriangleIDs(DynamicMesh, Context->TriangleIds[i], /*bHasTriangleIDGaps=*/Dummy);
+			NumIterations = Context->TriangleIds[i].List->Num();
+			break;
+		}
+		case EPCGMeshSamplingMethod::PoissonSampling:
+		{
+			// No preparation needed
+			break;
+		}
+		default:
+		{
+			PCGLog::LogErrorOnGraph(LOCTEXT("InvalidOperation", "Invalid sampling operation"), InContext);
+			return true;
+		}
+		}
+
+		if (NumIterations != -1)
+		{
+			Context->StartingIndices[i + 1] = Context->StartingIndices[i] + NumIterations;
+			CurrentOutPointData->GetMutablePoints().SetNumUninitialized(NumIterations);
+		}
+
+		auto CreateAttribute = [this, Context, Metadata = CurrentOutPointData->Metadata]<typename T>(const bool bShouldAdd, const FName AttributeName, T DefaultValue, bool bAllowInterpolation, TArray<FPCGMetadataAttribute<T>*>& OutAttributesArray, const TCHAR* What)
+		{
+			if (bShouldAdd)
+			{
+				FPCGMetadataAttribute<T>* Attribute = Metadata->CreateAttribute<T>(AttributeName, std::move(DefaultValue), bAllowInterpolation, /*bOverrideParent=*/true);
+				if (!Attribute)
+				{
+					PCGE_LOG(Warning, GraphAndLog, FText::Format(LOCTEXT("AttributeCreationFailed", "Failed to create attribute {0} for {1}. {1} won't be computed"), FText::FromName(AttributeName), FText::FromString(What)));
+				}
+
+				OutAttributesArray.Emplace(Attribute);
+			}
+		};
+
+		// It's not clear how to compute UVs for Vertices as they are part of multiple triangles. So disable for this mode. Same for triangle ids.
+		if (Settings->SamplingMethod != EPCGMeshSamplingMethod::OnePointPerVertex)
+		{
+			CreateAttribute(Settings->bExtractUVAsAttribute, Settings->UVAttributeName, FVector2D::ZeroVector, /*bAllowInterpolation=*/true, Context->UVAttributes, TEXT("UVs"));
+			CreateAttribute(Settings->bOutputTriangleIds, Settings->TriangleIdAttributeName, -1, /*bAllowInterpolation=*/false, Context->TriangleIdAttributes, TEXT("Triangle IDs"));
+			CreateAttribute(Settings->bOutputMaterialInfo, Settings->MaterialIdAttributeName, -1, /*bAllowInterpolation=*/false, Context->MaterialIdAttributes, TEXT("Material IDs"));
+			CreateAttribute(Settings->bOutputMaterialInfo, Settings->MaterialAttributeName, FSoftObjectPath(), /*bAllowInterpolation=*/false, Context->MaterialAttributes, TEXT("Material"));
+		}
+
+		Outputs.Emplace_GetRef().Data = CurrentOutPointData;
+	}
+
+	Context->bDataPrepared = true;
+	return true;
+}
+
+bool FPCGMeshSamplerElement::ExecuteInternal(FPCGContext* InContext) const
+{
+	TRACE_CPUPROFILER_EVENT_SCOPE(FPCGMeshSamplerElement::Execute);
+	FPCGMeshSamplerContext* Context = static_cast<FPCGMeshSamplerContext*>(InContext);
+
+	check(Context);
+
+	if (!Context->bDataPrepared || Context->OutPointData.IsEmpty())
+	{
+		return true;
+	}
+
+	const UPCGMeshSamplerSettings* Settings = Context->GetInputSettings<UPCGMeshSamplerSettings>();
+	check(Settings);
+
+	bool bIsDone = false;
+
+	// Preparing the set function to extract the color to density.
+	auto SetPointDensityTo1 = [](const FLinearColor&, FPCGPoint& OutPoint){ OutPoint.Density = 1.0f; };
+	auto SetPointDensityToRed = [](const FLinearColor& Color, FPCGPoint& OutPoint){ OutPoint.Density = Color.R; };
+	auto SetPointDensityToGreen = [](const FLinearColor& Color, FPCGPoint& OutPoint){ OutPoint.Density = Color.G; };
+	auto SetPointDensityToBlue = [](const FLinearColor& Color, FPCGPoint& OutPoint){ OutPoint.Density = Color.B; };
+	auto SetPointDensityToAlpha = [](const FLinearColor& Color, FPCGPoint& OutPoint){ OutPoint.Density = Color.A; };
+
+	// Store it in a function pointer.
+	void(*SetPointDensityPtr)(const FLinearColor&, FPCGPoint&) = SetPointDensityTo1;
+
+	if (Settings->bUseColorChannelAsDensity)
+	{
+		switch (Settings->ColorChannelAsDensity)
+		{
+		case EPCGColorChannel::Red:
+			SetPointDensityPtr = SetPointDensityToRed;
+			break;
+		case EPCGColorChannel::Green:
+			SetPointDensityPtr = SetPointDensityToGreen;
+			break;
+		case EPCGColorChannel::Blue:
+			SetPointDensityPtr = SetPointDensityToBlue;
+			break;
+		case EPCGColorChannel::Alpha:
+			SetPointDensityPtr = SetPointDensityToAlpha;
+			break;
+		default:
+			checkNoEntry();
+			break;
+		}
+	}
+
+	switch (Settings->SamplingMethod)
+	{
+	case EPCGMeshSamplingMethod::OnePointPerVertex:
+	{
+		bIsDone = PCGMeshSampler::SampleOnePointPerVertex(Settings, Context, SetPointDensityPtr);
+		break;
+	}
+	case EPCGMeshSamplingMethod::OnePointPerTriangle:
+	{
+		bIsDone = PCGMeshSampler::SampleOnePointPerTriangle(Settings, Context, SetPointDensityPtr);
+		break;
+	}
+	case EPCGMeshSamplingMethod::PoissonSampling:
+	{
+		bIsDone = PCGMeshSampler::PoissonSampling(Settings, Context, SetPointDensityPtr);
+		break;
+	}
+	default:
+	{
+		// Already logged in the PrepareData, we should never arrive there
+		checkNoEntry();
+		return true;
+	}
+	}
+
+	return bIsDone;
+}
+
+#undef LOCTEXT_NAMESPACE
diff --git a/Engine/Plugins/PCGInterops/PCGGeometryScriptInterop/Source/PCGGeometryScriptInterop/Private/Elements/PCGSplineCrossSectionGenerator.cpp b/Engine/Plugins/PCGInterops/PCGGeometryScriptInterop/Source/PCGGeometryScriptInterop/Private/Elements/PCGSplineCrossSectionGenerator.cpp
new file mode 100644
--- /dev/null
+++ b/Engine/Plugins/PCGInterops/PCGGeometryScriptInterop/Source/PCGGeometryScriptInterop/Private/Elements/PCGSplineCrossSectionGenerator.cpp
@@ -0,0 +1,477 @@
+// Copyright Epic Games, Inc. All Rights Reserved.
+
+#include "Elements/PCGSplineCrossSectionGenerator.h"
+
+#include "Helpers/PCGGeometryHelpers.h"
+
+#include "PCGContext.h"
+#include "PCGComponent.h"
+#include "Data/PCGPrimitiveData.h"
+#include "Data/PCGSplineData.h"
+#include "Helpers/PCGHelpers.h"
+
+#include "UDynamicMesh.h"
+#include "Algo/Find.h"
+#include "Algo/ForEach.h"
+#include "Components/SplineComponent.h"
+#include "GeometryScript/MeshBooleanFunctions.h"
+#include "GeometryScript/MeshModelingFunctions.h"
+#include "GeometryScript/MeshQueryFunctions.h"
+#include "GeometryScript/MeshSelectionFunctions.h"
+#include "GeometryScript/MeshSelectionQueryFunctions.h"
+#include "GeometryScript/MeshSimplifyFunctions.h"
+#include "GeometryScript/SceneUtilityFunctions.h"
+
+#define LOCTEXT_NAMESPACE "PCGSplineCrossSectionGeneratorElement"
+
+namespace SplineCrossSectionGeneratorConstants
+{
+	const FName DefaultExtrusionVectorAttributeName(TEXT("ExtrusionVector"));
+	static const FGeometryScriptMeshPlaneCutOptions CutPlaneOptions{/*bFillHoles=*/false, /*bFillSpans=*/false, /*bFlipCutSide=*/true, 1};
+	static const FGeometryScriptCopyMeshFromComponentOptions CopyMeshFromComponentOptions{/*bWantNormals=*/false, /*bWantTangents=*/false};
+}
+
+struct FCrossSection
+{
+	int Tier;
+	double Height;
+	TArray<FVector> PointLocations;
+};
+
+UPCGSplineCrossSectionGeneratorSettings::UPCGSplineCrossSectionGeneratorSettings()
+{
+	ExtrusionVectorAttribute.SetAttributeName(SplineCrossSectionGeneratorConstants::DefaultExtrusionVectorAttributeName);
+}
+
+TArray<FPCGPinProperties> UPCGSplineCrossSectionGeneratorSettings::InputPinProperties() const
+{
+	TArray<FPCGPinProperties> PinProperties;
+	PinProperties.Emplace(PCGPinConstants::DefaultInputLabel, EPCGDataType::Primitive);
+	return PinProperties;
+}
+
+TArray<FPCGPinProperties> UPCGSplineCrossSectionGeneratorSettings::OutputPinProperties() const
+{
+	TArray<FPCGPinProperties> PinProperties;
+	PinProperties.Emplace(PCGPinConstants::DefaultOutputLabel, EPCGDataType::Spline);
+	return PinProperties;
+}
+
+FPCGElementPtr UPCGSplineCrossSectionGeneratorSettings::CreateElement() const
+{
+	return MakeShared<FPCGSplineCrossSectionGeneratorElement>();
+}
+
+bool FPCGSplineCrossSectionGeneratorElement::ExecuteInternal(FPCGContext* Context) const
+{
+	TRACE_CPUPROFILER_EVENT_SCOPE(FPCGSplineCrossSectionGeneratorElement::Execute);
+
+	const UPCGSplineCrossSectionGeneratorSettings* Settings = Context->GetInputSettings<UPCGSplineCrossSectionGeneratorSettings>();
+	const TArray<FPCGTaggedData> PrimitiveInputs = Context->InputData.GetInputsByPin(PCGPinConstants::DefaultInputLabel);
+
+	if (PrimitiveInputs.IsEmpty())
+	{
+		return true;
+	}
+
+	// TODO: Could also take in PCGVolume
+	TArray<const UPCGPrimitiveData*> PrimitiveDataCollection;
+	PrimitiveDataCollection.Reserve(PrimitiveInputs.Num());
+
+	for (const FPCGTaggedData& TaggedData : PrimitiveInputs)
+	{
+		if (const UPCGPrimitiveData* PrimitiveData = Cast<UPCGPrimitiveData>(TaggedData.Data))
+		{
+			PrimitiveDataCollection.Add(PrimitiveData);
+		}
+	}
+
+	if (PrimitiveDataCollection.IsEmpty())
+	{
+		return true;
+	}
+
+	UDynamicMesh* DynamicMesh = FPCGContext::NewObject_AnyThread<UDynamicMesh>(Context);
+#if WITH_EDITOR
+	UGeometryScriptDebug* DynamicMeshDebug = FPCGContext::NewObject_AnyThread<UGeometryScriptDebug>(Context);
+#else
+	UGeometryScriptDebug* DynamicMeshDebug = nullptr;
+#endif // WITH_EDITOR
+
+	// Collect all the primitives and append them to the dynamic mesh
+	for (const UPCGPrimitiveData* PrimitiveData : PrimitiveDataCollection)
+	{
+		TWeakObjectPtr<UPrimitiveComponent> PrimitivePtr = PrimitiveData->GetComponent();
+		if (!PrimitivePtr.IsValid())
+		{
+			continue;
+		}
+
+		// Convert from scene component to mesh and begin boolean operation
+		if (USceneComponent* SceneComponent = Cast<USceneComponent>(PrimitivePtr.Pin().Get()))
+		{
+			bool bFirstMesh = DynamicMesh->IsEmpty();
+
+			UDynamicMesh* ComponentMesh = NewObject<UDynamicMesh>();
+			EGeometryScriptOutcomePins Outcome;
+			FTransform PrimitiveTransform;
+			// TODO: If we have repeats, ie. components that match content/transform, we could skip them.
+			UGeometryScriptLibrary_SceneUtilityFunctions::CopyMeshFromComponent(
+				SceneComponent,
+				(bFirstMesh ? DynamicMesh : ComponentMesh),
+				SplineCrossSectionGeneratorConstants::CopyMeshFromComponentOptions,
+				/*bTransformToWorld=*/true,
+				PrimitiveTransform,
+				Outcome,
+				DynamicMeshDebug);
+
+			// TODO: Investigate if merging first and boolean the final result would be faster.
+			if (!bFirstMesh)
+			{
+				UGeometryScriptLibrary_MeshBooleanFunctions::ApplyMeshBoolean(
+					DynamicMesh,
+					FTransform::Identity,
+					ComponentMesh,
+					FTransform::Identity,
+					EGeometryScriptBooleanOperation::Union,
+					FGeometryScriptMeshBooleanOptions(), // TODO: check defaults
+					DynamicMeshDebug);
+			}
+		}
+	}
+
+	// Reduce vertex count by simplifying coplanar triangles. Also removes index gaps.
+	UGeometryScriptLibrary_MeshSimplifyFunctions::ApplySimplifyToPlanar(DynamicMesh, FGeometryScriptPlanarSimplifyOptions(), DynamicMeshDebug);
+
+	FGeometryScriptVectorList VertexList;
+	bool bHasGapsDummy = false;
+	UGeometryScriptLibrary_MeshQueryFunctions::GetAllVertexPositions(DynamicMesh, VertexList, /*bSkipGaps=*/false, bHasGapsDummy);
+
+	if (!VertexList.List.IsValid() || VertexList.List->IsEmpty())
+	{
+		return true;
+	}
+
+	// No cross-sections will be generated. Early out.
+	if (!Settings->bEnableFindFeatures && !Settings->bEnableTierSlicing)
+	{
+		return true;
+	}
+
+	// A safeguard to prevent analysis on a significantly large mesh.
+	if (Settings->bEnableFindFeatures && VertexList.List->Num() > Settings->MaxMeshVertexCount)
+	{
+		PCGLog::LogWarningOnGraph(LOCTEXT("MaxMeshVertexCount", "Combined mesh exceeds Max Mesh Vertex Count and will be skipped."), Context);
+		return true;
+	}
+
+	FVector SliceOrigin = FVector::ZeroVector;
+	const FVector SliceDirection = Settings->SliceDirection.GetSafeNormal();
+
+	double MinProjectionScalar = std::numeric_limits<double>::max();
+	double MaxProjectionScalar = std::numeric_limits<double>::min();
+
+	// Project vertices along the slice direction and find the minimum height
+	TArray<double> ProjectionScalars;
+	ProjectionScalars.Reserve(VertexList.List->Num());
+	for (const FVector& Vertex : *VertexList.List)
+	{
+		// Implementation note: since the direction vector is guaranteed to be normalized, we can use the projection scalar as distance
+		double ProjectionScalar = FVector::DotProduct(Vertex, SliceDirection);
+		ProjectionScalars.Emplace(ProjectionScalar);
+
+		// Find the min and max projected "height" to define the slicing extents
+		if (MinProjectionScalar > ProjectionScalar)
+		{
+			MinProjectionScalar = ProjectionScalar;
+			// Find the vertex for which we'll start slicing
+			SliceOrigin = Vertex;
+		}
+
+		MaxProjectionScalar = FMath::Max(MaxProjectionScalar, ProjectionScalar);
+	}
+
+	// Convert from projection scalar to height from origin
+	TMap<double, int> HeightMap;
+	for (double ProjectionScalar : ProjectionScalars)
+	{
+		// TODO: It would be better to compare an epsilon, in the case of sub 1cm meshes.
+		double Height = FMath::Floor(ProjectionScalar - MinProjectionScalar);
+		++HeightMap.FindOrAdd(Height, 0);
+	}
+
+	TArray<double> HeightKeys;
+	HeightMap.GetKeys(HeightKeys);
+
+	TArray<double, TInlineAllocator<16>> TierHeights;
+	// Find heights with coplanar vertices using the height counts
+	if (Settings->bEnableFindFeatures)
+	{
+		Algo::CopyIf(HeightKeys, TierHeights, [&HeightMap, MinVertices = Settings->MinimumCoplanarVertices](double Key)
+		{
+			return HeightMap[Key] >= MinVertices;
+		});
+	}
+
+	// Slice from the base of the mesh to the highest point by slice direction at the given interval
+	if (Settings->bEnableTierSlicing)
+	{
+		double TierSlicingResolution = FMath::Max(Settings->TierSlicingResolution, 100.0);
+		for (double Height = 0; Height <= (MaxProjectionScalar - MinProjectionScalar); Height += TierSlicingResolution)
+		{
+			TierHeights.AddUnique(Height);
+		}
+	}
+
+	if (TierHeights.IsEmpty())
+	{
+		return true;
+	}
+
+	TierHeights.Sort([](double LHS, double RHS) { return LHS < RHS; });
+
+	if (Settings->bEnableTierMerging)
+	{
+		// Filter the heights by threshold
+		TArray<double, TInlineAllocator<16>> FilteredTiers;
+		FilteredTiers.Reserve(TierHeights.Num());
+		FilteredTiers.Add(TierHeights[0]);
+
+		// TODO: This currently always merges "down", but it could also alternatively blend with an average, merge "up", etc
+		for (int TierIndex = 1; TierIndex < TierHeights.Num(); ++TierIndex)
+		{
+			if (TierHeights[TierIndex] - FilteredTiers.Last() > Settings->TierMergingThreshold)
+			{
+				FilteredTiers.Add(TierHeights[TierIndex]);
+			}
+		}
+
+		TierHeights = MoveTemp(FilteredTiers);
+	}
+
+	// Set up the slice plane for GeometryScript, represented by an FTransform
+	FTransform SlicePlaneTransform;
+	SlicePlaneTransform.SetLocation(SliceOrigin);
+	FQuat AdjustedRotation = FQuat::FindBetween(FVector::UpVector, SliceDirection);
+	SlicePlaneTransform.SetRotation(AdjustedRotation);
+
+	TArray<FCrossSection> CrossSections;
+	for (int TierIndex = 0; TierIndex < TierHeights.Num() - 1; ++TierIndex)
+	{
+		FVector SliceLocation = SliceOrigin + SliceDirection * TierHeights[TierIndex];
+		SlicePlaneTransform.SetLocation(SliceLocation);
+
+		// Cuts the mesh at the specified plane, leaving a hole
+		UGeometryScriptLibrary_MeshBooleanFunctions::ApplyMeshPlaneCut(DynamicMesh, SlicePlaneTransform, SplineCrossSectionGeneratorConstants::CutPlaneOptions, DynamicMeshDebug);
+
+		// Break now as there is no point in continuing along this invalid mesh
+		if (!ensure(DynamicMesh->GetTriangleCount() >= 1))
+		{
+			PCGLog::LogWarningOnGraph(FText::Format(LOCTEXT("DynamicMeshInvalid", "Dynamic Mesh Invalid at tier: {0}"), FText::AsNumber(TierIndex)), Context);
+			break;
+		}
+
+		// The cut algorithm does not result in simplified planes, so simplify it
+		UGeometryScriptLibrary_MeshSimplifyFunctions::ApplySimplifyToPlanar(DynamicMesh, FGeometryScriptPlanarSimplifyOptions(), DynamicMeshDebug); // TODO: verify defaults
+
+		// TODO: Investigate if there's a simpler API call to get the cut plane triangles (vertices)
+		// Since select box is AABB, in order to select a slice, select with a plane just before and subtract the one just after
+		FGeometryScriptMeshSelection CurrentSelection{};
+		UGeometryScriptLibrary_MeshSelectionFunctions::SelectMeshElementsWithPlane(
+			DynamicMesh,
+			CurrentSelection,
+			SliceLocation - (SliceDirection * UE_DOUBLE_KINDA_SMALL_NUMBER), // Move back a bit
+			SliceDirection,
+			EGeometryScriptMeshSelectionType::Vertices);
+
+		FGeometryScriptMeshSelection ExclusiveSelection{};
+		UGeometryScriptLibrary_MeshSelectionFunctions::SelectMeshElementsWithPlane(
+			DynamicMesh,
+			ExclusiveSelection,
+			SliceLocation + (SliceDirection * UE_DOUBLE_KINDA_SMALL_NUMBER), // Move forward a bit
+			SliceDirection,
+			EGeometryScriptMeshSelectionType::Vertices);
+
+		UGeometryScriptLibrary_MeshSelectionFunctions::CombineMeshSelections(
+			CurrentSelection,
+			ExclusiveSelection,
+			CurrentSelection,
+			EGeometryScriptCombineSelectionMode::Subtract);
+
+		TArray<FGeometryScriptIndexList> VertexIndexLists;
+		TArray<FGeometryScriptPolyPath> PolyPaths;
+		int NumLoops;
+		bool bFoundErrors;
+		// TODO: Find a more direct way to determine multiple poly paths--one for each island--and order the vertices
+		// Implementation note: Using Boundary Loops to create multiple 2D poly paths and order points.
+		UGeometryScriptLibrary_MeshSelectionQueryFunctions::GetMeshSelectionBoundaryLoops(
+			DynamicMesh,
+			CurrentSelection,
+			VertexIndexLists,
+			PolyPaths,
+			NumLoops,
+			bFoundErrors,
+			DynamicMeshDebug);
+
+#if WITH_EDITOR
+		// At this point, we're done with GeometryScript for this iteration. Print errors if they occur.
+		PCGGeometryHelpers::GeometryScriptDebugToPCGLog(Context, DynamicMeshDebug);
+#endif // WITH_EDITOR
+
+		if (NumLoops < 1 || PolyPaths.IsEmpty())
+		{
+			continue;
+		}
+
+		// Shoelace formula - https://en.wikipedia.org/wiki/Shoelace_formula.
+		// Copy PointLocations in by value to be used as a temporary container.
+		auto ComputePolyPathArea = [&AdjustedRotation](TArray<FVector> PointLocations)
+		{
+			const int NumPoints = PointLocations.Num();
+			if (NumPoints < 3)
+			{
+				return 0.0;
+			}
+
+			// Transform the points to 2D referential in order to calculate the surface area
+			Algo::ForEach(PointLocations, [&AdjustedRotation](FVector& Location) { Location = AdjustedRotation.Inverse().RotateVector(Location); });
+
+			double Area = 0.0;
+			for (int PointIndex = 0; PointIndex < PointLocations.Num() - 1; ++PointIndex)
+			{
+				Area += PointLocations[PointIndex].X * PointLocations[PointIndex + 1].Y - PointLocations[PointIndex + 1].X * PointLocations[PointIndex].Y;
+			}
+
+			// Calculate the last term (last and first point) separately
+			Area += PointLocations[NumPoints - 1].X * PointLocations[0].Y - PointLocations[0].X * PointLocations[NumPoints - 1].Y;
+
+			// The sign of the area could be used to determine winding, if needed in the future
+			return FMath::Abs(Area) * 0.5;
+		};
+
+		for (int LoopIndex = 0; LoopIndex < NumLoops; ++LoopIndex)
+		{
+			if (!ensure(PolyPaths[LoopIndex].bClosedLoop))
+			{
+				continue;
+			}
+
+			// Do a planar check of one of the vertices and eliminate poly paths outside our cut plane
+			// TODO: This is a crutch to keep the boundary loops from creeping.
+			double DotProduct = FMath::Abs(FVector::DotProduct(PolyPaths[LoopIndex].Path->Last() - SliceLocation, SliceDirection));
+			if (DotProduct > UE_DOUBLE_KINDA_SMALL_NUMBER)
+			{
+				continue;
+			}
+
+			if (Settings->bEnableMinAreaCulling && ComputePolyPathArea(*PolyPaths[LoopIndex].Path) <= Settings->MinAreaCullingThreshold * 100.0)
+			{
+				continue;
+			}
+
+			FCrossSection& CurrentTier = CrossSections.Emplace_GetRef();
+			CurrentTier.Tier = TierIndex;
+			CurrentTier.Height = TierHeights[TierIndex + 1] - TierHeights[TierIndex];
+			CurrentTier.PointLocations.Append(*PolyPaths[LoopIndex].Path);
+		}
+	}
+
+	// Filter unintended excess cross-sections
+	if (Settings->bRemoveRedundantSections)
+	{
+		TArray<int, TInlineAllocator<16>> RedundantSections;
+		// Implementation note: the cross-sections are sorted by height here, so the highest can be taken
+		CrossSections.Sort([](const FCrossSection& LHS, const FCrossSection& RHS) { return LHS.Tier < RHS.Tier; });
+
+		for (int FirstIndex = 0; FirstIndex < CrossSections.Num(); ++FirstIndex)
+		{
+			if (RedundantSections.Contains(FirstIndex))
+			{
+				continue;
+			}
+
+			for (int SecondIndex = CrossSections.Num() - 1; SecondIndex > FirstIndex; --SecondIndex)
+			{
+				if (RedundantSections.Contains(SecondIndex))
+				{
+					continue;
+				}
+
+				FCrossSection& FirstSection = CrossSections[FirstIndex];
+				const FCrossSection& SecondSection = CrossSections[SecondIndex];
+
+				if (SecondSection.Tier < FirstSection.Tier)
+				{
+					break;
+				}
+
+				// Disqualify based on point count
+				if (FirstSection.PointLocations.Num() != SecondSection.PointLocations.Num())
+				{
+					continue;
+				}
+
+				// Check if the cross-section's vertices are all equal to find if it's redundant for culling.
+				// Implementation Note: This behavior is not always intuitive with any mesh and will cull even if there are tiers in between.
+				// TODO: Some future options might include checking for intermediate tiers between, etc. However, we need to be mindful of merge order.
+				bool bCrossSectionIsEqual = true;
+				for (int PointIndex = 0; PointIndex < FirstSection.PointLocations.Num(); ++PointIndex)
+				{
+					if (!Algo::FindByPredicate(SecondSection.PointLocations, [FirstLocation = FirstSection.PointLocations[PointIndex]](const FVector& SecondLocation)
+					{
+						return FMath::IsNearlyEqual(FirstLocation.X, SecondLocation.X) && FMath::IsNearlyEqual(FirstLocation.Y, SecondLocation.Y);
+					}))
+					{
+						bCrossSectionIsEqual = false;
+						break;
+					}
+				}
+
+				if (bCrossSectionIsEqual)
+				{
+					RedundantSections.Add(SecondIndex);
+					// Combine the heights of both sections
+					FirstSection.Height += SecondSection.Height;
+				}
+			}
+		}
+
+		// Sort inverse, because they will be remove swapped from the actual array
+		RedundantSections.Sort([](int LHS, int RHS) { return LHS > RHS; });
+
+		for (int Index : RedundantSections)
+		{
+			CrossSections.RemoveAtSwap(Index);
+		}
+	}
+
+	// Create the splines from the tiers
+	for (const FCrossSection& CrossSection : CrossSections)
+	{
+		TArray<FSplinePoint> SplinePoints;
+		SplinePoints.Reserve(CrossSection.PointLocations.Num());
+		FPCGTaggedData& OutputData = Context->OutputData.TaggedData.Emplace_GetRef();
+
+		for (int Index = 0; Index < CrossSection.PointLocations.Num(); ++Index)
+		{
+			SplinePoints.Emplace(static_cast<float>(Index),
+				CrossSection.PointLocations[Index],
+				FVector::ZeroVector,
+				FVector::ZeroVector,
+				FRotator::ZeroRotator,
+				FVector::OneVector,
+				ESplinePointType::Linear);
+		}
+
+		UPCGSplineData* OutSplineData = FPCGContext::NewObject_AnyThread<UPCGSplineData>(Context);
+		OutSplineData->Initialize(SplinePoints, true, FTransform::Identity);
+		OutputData.Data = OutSplineData;
+
+		check(OutSplineData->Metadata);
+		OutSplineData->Metadata->CreateAttribute(Settings->ExtrusionVectorAttribute.GetName(), FVector(0, 0, CrossSection.Height), /*bAllowsInterpolation=*/false, /*bOverridesParent=*/true);
+	}
+
+	return true;
+}
+
+#undef LOCTEXT_NAMESPACE
diff --git a/Engine/Plugins/PCGInterops/PCGGeometryScriptInterop/Source/PCGGeometryScriptInterop/Private/Helpers/PCGGeometryHelpers.cpp b/Engine/Plugins/PCGInterops/PCGGeometryScriptInterop/Source/PCGGeometryScriptInterop/Private/Helpers/PCGGeometryHelpers.cpp
new file mode 100644
--- /dev/null
+++ b/Engine/Plugins/PCGInterops/PCGGeometryScriptInterop/Source/PCGGeometryScriptInterop/Private/Helpers/PCGGeometryHelpers.cpp
@@ -0,0 +1,27 @@
+// Copyright Epic Games, Inc. All Rights Reserved.
+
+#include "Helpers/PCGGeometryHelpers.h"
+
+#include "Utils/PCGLogErrors.h"
+
+#include "GeometryScript/GeometryScriptTypes.h"
+
+namespace PCGGeometryHelpers
+{
+	void GeometryScriptDebugToPCGLog(FPCGContext* Context, const UGeometryScriptDebug* Debug)
+	{
+		check(Context && Debug);
+
+		for (const FGeometryScriptDebugMessage& Message : Debug->Messages)
+		{
+			if (Message.MessageType == EGeometryScriptDebugMessageType::ErrorMessage)
+			{
+				PCGLog::LogErrorOnGraph(Message.Message, Context);
+			}
+			else
+			{
+				PCGLog::LogWarningOnGraph(Message.Message, Context);
+			}
+		}
+	}
+}
\ No newline at end of file
diff --git a/Engine/Plugins/PCGInterops/PCGGeometryScriptInterop/Source/PCGGeometryScriptInterop/Private/PCGGeometryScriptInteropModule.cpp b/Engine/Plugins/PCGInterops/PCGGeometryScriptInterop/Source/PCGGeometryScriptInterop/Private/PCGGeometryScriptInteropModule.cpp
new file mode 100644
--- /dev/null
+++ b/Engine/Plugins/PCGInterops/PCGGeometryScriptInterop/Source/PCGGeometryScriptInterop/Private/PCGGeometryScriptInteropModule.cpp
@@ -0,0 +1,16 @@
+// Copyright Epic Games, Inc. All Rights Reserved.
+
+#include "Modules/ModuleManager.h"
+
+class FPCGGeometryScriptInteropModule final : public IModuleInterface
+{
+public:
+	//~ IModuleInterface implementation
+	virtual bool SupportsDynamicReloading() override
+	{
+		return true;
+	}
+	//~ End IModuleInterface implementation
+};
+
+IMPLEMENT_MODULE(FPCGGeometryScriptInteropModule, PCGGeometryScriptInterop);
diff --git a/Engine/Plugins/PCGInterops/PCGGeometryScriptInterop/Source/PCGGeometryScriptInterop/Public/Elements/PCGMeshSampler.h b/Engine/Plugins/PCGInterops/PCGGeometryScriptInterop/Source/PCGGeometryScriptInterop/Public/Elements/PCGMeshSampler.h
new file mode 100644
--- /dev/null
+++ b/Engine/Plugins/PCGInterops/PCGGeometryScriptInterop/Source/PCGGeometryScriptInterop/Public/Elements/PCGMeshSampler.h
@@ -0,0 +1,242 @@
+// Copyright Epic Games, Inc. All Rights Reserved.
+
+#pragma once
+
+#include "PCGContext.h"
+#include "PCGElement.h"
+#include "PCGSettings.h"
+#include "Elements/PCGLoadObjectsContext.h"
+
+#include "Async/Future.h"
+#include "GeometryScript/GeometryScriptTypes.h"
+#include "GeometryScript/MeshSamplingFunctions.h"
+#include "Templates/UniquePtr.h"
+#include "UObject/ObjectPtr.h"
+#include "UObject/SoftObjectPtr.h"
+
+#include <atomic>
+
+#include "PCGMeshSampler.generated.h"
+
+template <typename T>
+class FPCGMetadataAttribute;
+
+class FProgressCancel;
+class UDynamicMesh;
+class UPCGPointData;
+class UStaticMesh;
+
+UENUM(BlueprintType)
+enum class EPCGMeshSamplingMethod : uint8
+{
+	/** Sample one point (at the center) of each triangle of the mesh. */
+	OnePointPerTriangle,
+
+	/** Sample one point per vertex on the mesh. */
+	OnePointPerVertex,
+
+	/** Use Poisson sampling to sample points on the mesh. Can be expensive and therefore it is not framebound. */
+	PoissonSampling
+};
+
+UENUM()
+enum class EPCGColorChannel
+{
+	Red,
+	Green,
+	Blue,
+	Alpha
+};
+
+/**
+* Sample points on a mesh.
+*/
+UCLASS(BlueprintType, ClassGroup = (Procedural))
+class PCGGEOMETRYSCRIPTINTEROP_API UPCGMeshSamplerSettings : public UPCGSettings
+{
+	GENERATED_BODY()
+
+public:
+	UPCGMeshSamplerSettings();
+
+	//~Begin UObject interface
+	virtual void PostLoad() override;
+	//~End UObject interface
+
+	//~Begin UPCGSettings interface
+#if WITH_EDITOR
+	virtual FName GetDefaultNodeName() const override { return FName(TEXT("MeshSampler")); }
+	virtual FText GetDefaultNodeTitle() const override { return NSLOCTEXT("PCGMeshSamplerSettings", "NodeTitle", "Mesh Sampler"); }
+	virtual FText GetNodeTooltipText() const override;
+	virtual EPCGSettingsType GetType() const override { return EPCGSettingsType::Sampler; }
+#endif
+
+	virtual TArray<FPCGPinProperties> InputPinProperties() const override;
+	virtual TArray<FPCGPinProperties> OutputPinProperties() const override;
+
+protected:
+	virtual FPCGElementPtr CreateElement() const override;
+	// ~End UPCGSettings interface
+
+public:
+	/** 
+	* Can provide a list of inputs to sample the meshes from. It can be a list of StaticMeshes, a list of Actors that have a Scene Component (like a Static Mesh Component), or a list of Scene Components directly.
+	* Geometry Script needs to be able to extract a dynamic mesh from this scene component (so won't work for ISMCs for example) and for now will work only with a single scene component.
+	* Each entry (either in the same data or seperate data) will produce a unique output data.
+	*/
+	UPROPERTY(BlueprintReadWrite, EditAnywhere, Category = "Settings")
+	bool bExtractMeshFromInput = false;
+
+	/** Soft Object Path to the mesh to sample from. Will be loaded. */
+	UPROPERTY(BlueprintReadWrite, EditAnywhere, Category = "Settings", meta = (PCG_Overridable, EditCondition = "!bExtractMeshFromInput", EditConditionHides, ShowAfter = "bExtractMeshFromInput"))
+	TSoftObjectPtr<UStaticMesh> StaticMesh;
+
+	/** Selector to read data from. */
+	UPROPERTY(BlueprintReadWrite, EditAnywhere, Category = "Settings", meta = (PCG_Overridable, EditCondition = "bExtractMeshFromInput", EditConditionHides, ShowAfter = "bExtractMeshFromInput"))
+	FPCGAttributePropertyInputSelector InputSource;
+
+	UPROPERTY(BlueprintReadWrite, EditAnywhere, Category = "Settings")
+	EPCGMeshSamplingMethod SamplingMethod = EPCGMeshSamplingMethod::OnePointPerTriangle;
+
+	/** Will extract the color channel into the density. */
+	UPROPERTY(BlueprintReadWrite, EditAnywhere, Category = "Settings|Color & Density", meta = (PCG_Overridable, InlineEditConditionToggle, PCG_OverrideAliases="bUseRedAsDensity"))
+	bool bUseColorChannelAsDensity = false;
+	
+	UPROPERTY(BlueprintReadWrite, EditAnywhere, Category = "Settings|Color & Density", meta = (PCG_Overridable, EditCondition="bUseColorChannelAsDensity"))
+	EPCGColorChannel ColorChannelAsDensity = EPCGColorChannel::Red;
+
+	/** Enable voxelisation as a preparation pass. Can be more expensive given the VoxelSize. */
+	UPROPERTY(BlueprintReadWrite, EditAnywhere, Category = "Settings|Voxelize Options", meta = (PCG_Overridable))
+	bool bVoxelize = false;
+
+	/** Size of a voxel for the voxelization. */
+	UPROPERTY(BlueprintReadWrite, EditAnywhere, Category = "Settings|Voxelize Options", meta = (PCG_Overridable, EditCondition = "bVoxelize"))
+	float VoxelSize = 100.0f;
+
+	/** Post-processing pass after voxelization to remove hidden triangles. */
+	UPROPERTY(BlueprintReadWrite, EditAnywhere, Category = "Settings|Voxelize Options", meta = (PCG_Overridable, EditCondition = "bVoxelize"))
+	bool bRemoveHiddenTriangles = true;
+
+	/** LOD type to use when creating DynamicMesh from specified StaticMesh. */
+	UPROPERTY(BlueprintReadWrite, EditAnywhere, Category = "Settings|LODSettings", meta = (PCG_Overridable))
+	EGeometryScriptLODType RequestedLODType = EGeometryScriptLODType::RenderData;
+
+	UPROPERTY(BlueprintReadWrite, EditAnywhere, Category = "Settings|LODSettings", meta = (PCG_Overridable))
+	int32 RequestedLODIndex = 0;
+
+	// Poisson Sampling parameters
+	UPROPERTY(BlueprintReadWrite, EditAnywhere, Category = "Settings|Poisson sampling", meta = (PCG_Overridable, EditCondition = "SamplingMethod == EPCGMeshSamplingMethod::PoissonSampling"))
+	FGeometryScriptMeshPointSamplingOptions SamplingOptions;
+
+	UPROPERTY(BlueprintReadWrite, EditAnywhere, Category = "Settings|Poisson sampling", meta = (PCG_Overridable, EditCondition = "SamplingMethod == EPCGMeshSamplingMethod::PoissonSampling"))
+	FGeometryScriptNonUniformPointSamplingOptions NonUniformSamplingOptions;
+
+	UPROPERTY(BlueprintReadWrite, EditAnywhere, Category = "Settings|UVs", meta = (PCG_Overridable, EditCondition = "SamplingMethod != EPCGMeshSamplingMethod::OnePointPerVertex"))
+	bool bExtractUVAsAttribute = false;
+
+	UPROPERTY(BlueprintReadWrite, EditAnywhere, Category = "Settings|UVs", meta = (PCG_Overridable, EditCondition = "SamplingMethod != EPCGMeshSamplingMethod::OnePointPerVertex && bExtractUVAsAttribute", EditConditionHides))
+	FName UVAttributeName = TEXT("UV");
+
+	UPROPERTY(BlueprintReadWrite, EditAnywhere, Category = "Settings|UVs", meta = (PCG_Overridable, EditCondition = "SamplingMethod != EPCGMeshSamplingMethod::OnePointPerVertex && bExtractUVAsAttribute", EditConditionHides))
+	int32 UVChannel = 0;
+
+	UPROPERTY(BlueprintReadWrite, EditAnywhere, Category = "Settings|Extra", meta = (PCG_Overridable, EditCondition = "SamplingMethod != EPCGMeshSamplingMethod::OnePointPerVertex"))
+	bool bOutputTriangleIds = false;
+
+	UPROPERTY(BlueprintReadWrite, EditAnywhere, Category = "Settings|Extra", meta = (PCG_Overridable, EditCondition = "SamplingMethod != EPCGMeshSamplingMethod::OnePointPerVertex && bOutputTriangleIds", EditConditionHides))
+	FName TriangleIdAttributeName = TEXT("TriangleId");
+
+	UPROPERTY(BlueprintReadWrite, EditAnywhere, Category = "Settings|Extra", meta = (PCG_Overridable, EditCondition = "SamplingMethod != EPCGMeshSamplingMethod::OnePointPerVertex"))
+	bool bOutputMaterialInfo = false;
+
+	UPROPERTY(BlueprintReadWrite, EditAnywhere, Category = "Settings|Extra", meta = (PCG_Overridable, EditCondition = "SamplingMethod != EPCGMeshSamplingMethod::OnePointPerVertex && bOutputMaterialInfo", EditConditionHides))
+	FName MaterialIdAttributeName = TEXT("MaterialId");
+
+	UPROPERTY(BlueprintReadWrite, EditAnywhere, Category = "Settings|Extra", meta = (PCG_Overridable, EditCondition = "SamplingMethod != EPCGMeshSamplingMethod::OnePointPerVertex && bOutputMaterialInfo", EditConditionHides))
+	FName MaterialAttributeName = TEXT("Material");
+
+	/** Each PCG point represents a discretized, volumetric region of world space. The points' Steepness value [0.0 to
+	 * 1.0] establishes how "hard" or "soft" that volume will be represented. From 0, it will ramp up linearly
+	 * increasing its influence over the density from the point's center to up to two times the bounds. At 1, it will
+	 * represent a binary box function with the size of the point's bounds.
+	 */
+	UPROPERTY(BlueprintReadWrite, EditAnywhere, Category = "Settings|Points", meta=(ClampMin="0", ClampMax="1", PCG_Overridable))
+	float PointSteepness = 0.5f;
+
+	UPROPERTY(BlueprintReadWrite, EditAnywhere, Category = "Settings|Debug")
+	bool bSynchronousLoad = false;
+	
+protected:
+#if WITH_EDITORONLY_DATA
+	// Deprecated in UE 5.3 in favor of StaticMesh
+	UPROPERTY()
+	FSoftObjectPath StaticMeshPath_DEPRECATED;
+
+	UPROPERTY()
+	bool bUseRedAsDensity_DEPRECATED = false;
+#endif
+};
+
+/**
+* Extra context to store all the data that need to be kept between multiple executions (time slicing)
+*/
+struct FPCGMeshSamplerContext : public FPCGLoadObjectsFromPathContext
+{
+public:
+	~FPCGMeshSamplerContext();
+
+	using SetPointDensityFunc = void(*)(const FLinearColor&, FPCGPoint&);
+
+	void SetAttributeValues(int32 UVChannel, int32 TriangleId, const FVector& BarycentricCoord, FPCGPoint& OutPoint, int32 DataIndex);
+	void SetPointColorAndDensity(SetPointDensityFunc SetPointDensityFuncPtr, int32 TriangleId, const FVector& BarycentricCoord, FPCGPoint& OutPoint, int32 DataIndex);
+
+protected:
+	virtual void AddExtraStructReferencedObjects(FReferenceCollector& Collector) override;
+	
+public:
+	// Dynamic meshes. Will be added to root.
+	TArray<TObjectPtr<UDynamicMesh>> DynamicMeshes;
+
+	// Lists extracted from the meshes
+	TArray<FGeometryScriptVectorList> Positions;
+	TArray<FGeometryScriptColorList> Colors;
+	TArray<FGeometryScriptVectorList> Normals;
+	TArray<FGeometryScriptIndexList> TriangleIds;
+
+	// Output point data.
+	TArray<UPCGPointData*> OutPointData;
+
+	// Optional attributes.
+	TArray<FPCGMetadataAttribute<FVector2D>*> UVAttributes;
+	TArray<FPCGMetadataAttribute<int32>*> TriangleIdAttributes;
+	TArray<FPCGMetadataAttribute<int32>*> MaterialIdAttributes;
+	TArray<FPCGMetadataAttribute<FSoftObjectPath>*> MaterialAttributes;
+
+	// Material Specific
+	TArray<TArray<UMaterialInterface*>> ComponentMaterialList;
+	TArray<TArray<UMaterialInterface*>> AssetMaterialList;
+
+	// For Poisson sampling, we are starting futures that are not framebound
+	// Store the futures and synchronisation items in the context
+	TArray<TFuture<bool>> SamplingFutures;
+	std::atomic<bool> StopSampling = false;
+	TUniquePtr<FProgressCancel> SamplingProgess;
+
+	// Starting indices for each different object to sample. If we have 3 meshes of 20 elements (like vertices) each, the array will be [0, 20, 40, 60]
+	TArray<int32> StartingIndices;
+
+	// Set to true if prepared data succeeded.
+	bool bDataPrepared = false;
+};
+
+class FPCGMeshSamplerElement : public IPCGElement
+{
+public:
+	// Loading needs to be done on the main thread and accessing objects outside of PCG might not be thread safe, so taking the safe approach
+	virtual bool CanExecuteOnlyOnMainThread(FPCGContext* Context) const override { return true; }
+
+protected:
+	virtual FPCGContext* CreateContext() override;
+	virtual bool PrepareDataInternal(FPCGContext* InContext) const override;
+	virtual bool ExecuteInternal(FPCGContext* InContext) const override;
+};
\ No newline at end of file
diff --git a/Engine/Plugins/PCGInterops/PCGGeometryScriptInterop/Source/PCGGeometryScriptInterop/Public/Elements/PCGSplineCrossSectionGenerator.h b/Engine/Plugins/PCGInterops/PCGGeometryScriptInterop/Source/PCGGeometryScriptInterop/Public/Elements/PCGSplineCrossSectionGenerator.h
new file mode 100644
--- /dev/null
+++ b/Engine/Plugins/PCGInterops/PCGGeometryScriptInterop/Source/PCGGeometryScriptInterop/Public/Elements/PCGSplineCrossSectionGenerator.h
@@ -0,0 +1,89 @@
+// Copyright Epic Games, Inc. All Rights Reserved.
+
+#pragma once
+
+#include "PCGSettings.h"
+
+#include "PCGSplineCrossSectionGenerator.generated.h"
+
+/**
+ * Creates a spline cross-section of one more primitives based on vertex features, slice resolution, or both.
+ */
+UCLASS(BlueprintType, ClassGroup = (Procedural))
+class UPCGSplineCrossSectionGeneratorSettings : public UPCGSettings
+{
+	GENERATED_BODY()
+
+public:
+	UPCGSplineCrossSectionGeneratorSettings();
+
+	//~Begin UPCGSettings interface
+#if WITH_EDITOR
+	virtual FName GetDefaultNodeName() const override { return FName(TEXT("SplineCrossSectionGenerator")); }
+	virtual FText GetDefaultNodeTitle() const override { return NSLOCTEXT("PCGSplineCrossSectionGeneratorElement", "NodeTitle", "Spline Cross Section Generator"); }
+	virtual FText GetNodeTooltipText() const override { return NSLOCTEXT("PCGSplineCrossSectionGeneratorElement", "NodeTooltip", "Creates a spline cross-section of one more primitives based on vertex features, slice resolution, or both."); }
+	virtual EPCGSettingsType GetType() const override { return EPCGSettingsType::Spatial; }
+#endif
+
+protected:
+	virtual TArray<FPCGPinProperties> InputPinProperties() const override;
+	virtual TArray<FPCGPinProperties> OutputPinProperties() const override;
+	virtual FPCGElementPtr CreateElement() const override;
+	//~End UPCGSettings interface
+
+public:
+	/** Slicing will happen from the minimum vertex along this direction vector (normalized). Useful primarily with the Tier Slicing Resolution. */
+	UPROPERTY(BlueprintReadWrite, EditAnywhere, Category = Settings, meta = (PCG_Overridable))
+	FVector SliceDirection = FVector::UpVector;
+
+	/** The attribute that will be populated with each cross-section's extrusion vector. */
+	UPROPERTY(BlueprintReadWrite, EditAnywhere, Category = Settings, meta = (PCG_Overridable))
+	FPCGAttributePropertyOutputSelector ExtrusionVectorAttribute;
+
+	/** Find primitive tier "features", which consist of a number of co-planar vertices. */
+	UPROPERTY(BlueprintReadWrite, EditAnywhere, Category = Settings, meta = (PCG_Overridable))
+	bool bEnableFindFeatures = true;
+
+	/** The minimum required number of vertices that must be co-planar in order to be considered a tier "feature". */
+	UPROPERTY(BlueprintReadWrite, EditAnywhere, Category = Settings, meta = (PCG_Overridable, EditCondition = "bEnableFindFeatures", EditConditionHides))
+	int32 MinimumCoplanarVertices = 3;
+
+	/** Create tiers at a specified resolution. */
+	UPROPERTY(BlueprintReadWrite, EditAnywhere, Category = Settings, meta = (PCG_Overridable, InlineEditConditionToggle))
+	bool bEnableTierSlicing = false;
+
+	/** A tier will be created at this interval (in cm). */
+	UPROPERTY(BlueprintReadWrite, EditAnywhere, Category = Settings, meta = (PCG_Overridable, EditCondition = "bEnableTierSlicing", ClampMin = "100"))
+	double TierSlicingResolution = 500.0;
+
+	/** Cull tiers that are within a specified threshold. */
+	UPROPERTY(BlueprintReadWrite, EditAnywhere, Category = Settings, meta = (PCG_Overridable, InlineEditConditionToggle))
+	bool bEnableTierMerging = false;
+
+	/** If a tier is within this distance (in cm) of the previous tier, it will be culled. */
+	UPROPERTY(BlueprintReadWrite, EditAnywhere, Category = Settings, meta = (PCG_Overridable, EditCondition = "bEnableTierMerging"))
+	double TierMergingThreshold = 100.0;
+
+	/** Cull tiers that have a surface area smaller than a specified threshold. */
+	UPROPERTY(BlueprintReadWrite, EditAnywhere, Category = Settings, meta = (PCG_Overridable, InlineEditConditionToggle))
+	bool bEnableMinAreaCulling = false;
+
+	/** If a tier is smaller in area than this threshold, it will be culled. */
+	UPROPERTY(BlueprintReadWrite, EditAnywhere, Category = Settings, meta = (PCG_Overridable, EditCondition = "bEnableMinAreaCulling"))
+	double MinAreaCullingThreshold = 100.0;
+
+	/** If multiple tiers can be combined into a single tier without affecting the contour, remove the redundant one. Note: This will currently cull even if there are other unique tiers in between. */
+	UPROPERTY(BlueprintReadWrite, EditAnywhere, Category = Settings, meta = (PCG_Overridable))
+	bool bRemoveRedundantSections = true;
+
+	/** A safeguard to prevent finding features on an overly complex mesh. */
+	UPROPERTY(BlueprintReadWrite, EditAnywhere, Category = "Settings|Advanced", meta = (PCG_Overridable, EditCondition = "bEnableFindFeatures", EditConditionHides))
+	int32 MaxMeshVertexCount = 2048;
+};
+
+class FPCGSplineCrossSectionGeneratorElement : public IPCGElement
+{
+protected:
+	virtual bool ExecuteInternal(FPCGContext* Context) const override;
+	virtual bool IsCacheable(const UPCGSettings* InSettings) const override { return false; }
+};
diff --git a/Engine/Plugins/PCGInterops/PCGGeometryScriptInterop/Source/PCGGeometryScriptInterop/Public/Helpers/PCGGeometryHelpers.h b/Engine/Plugins/PCGInterops/PCGGeometryScriptInterop/Source/PCGGeometryScriptInterop/Public/Helpers/PCGGeometryHelpers.h
new file mode 100644
--- /dev/null
+++ b/Engine/Plugins/PCGInterops/PCGGeometryScriptInterop/Source/PCGGeometryScriptInterop/Public/Helpers/PCGGeometryHelpers.h
@@ -0,0 +1,13 @@
+// Copyright Epic Games, Inc. All Rights Reserved.
+
+#pragma once
+
+#include "Elements/PCGMeshSampler.h"
+
+struct FPCGContext;
+class UGeometryScriptDebug;
+
+namespace PCGGeometryHelpers
+{
+	PCGGEOMETRYSCRIPTINTEROP_API void GeometryScriptDebugToPCGLog(FPCGContext* Context, const UGeometryScriptDebug* Debug);
+}
