[PCG] Add PCG Dynamic Mesh data + debug visualization + getter

* Add a new type UPCGDynamicMeshData
* Add a new managed component: UPCGDynamicMeshManagedComponent
* Add a new debug visualization for Dynamic Meshes
* Add a new entry to the GetActorDataRegistry to be used with a Get Dynamic Data getter.
* Add helper function to duplicate data for the each of extending Geometry Script nodes in Blueprints.

#rnx

diff --git a/Engine/Plugins/PCG/Source/PCG/Private/PCGModule.cpp b/Engine/Plugins/PCG/Source/PCG/Private/PCGModule.cpp
--- a/Engine/Plugins/PCG/Source/PCG/Private/PCGModule.cpp
+++ b/Engine/Plugins/PCG/Source/PCG/Private/PCGModule.cpp
@@ -23,7 +23,12 @@
 
 FPCGModule& FPCGModule::GetPCGModuleChecked()
 {
-	return FModuleManager::GetModuleChecked<FPCGModule>("PCG");
+	return FModuleManager::GetModuleChecked<FPCGModule>(TEXT("PCG"));
+}
+
+bool FPCGModule::IsPCGModuleLoaded()
+{
+	return FModuleManager::Get().IsModuleLoaded(TEXT("PCG"));
 }
 
 #if WITH_EDITOR
diff --git a/Engine/Plugins/PCG/Source/PCG/Public/PCGModule.h b/Engine/Plugins/PCG/Source/PCG/Public/PCGModule.h
--- a/Engine/Plugins/PCG/Source/PCG/Public/PCGModule.h
+++ b/Engine/Plugins/PCG/Source/PCG/Public/PCGModule.h
@@ -38,7 +38,9 @@ public:
 	PCG_API static FPCGModule& GetPCGModuleChecked();
 	static const FPCGGetDataFunctionRegistry& ConstGetDataFunctionRegistry() { return GetPCGModuleChecked().GetDataFunctionRegistry; }
 	static FPCGGetDataFunctionRegistry& MutableGetDataFunctionRegistry() { return GetPCGModuleChecked().GetDataFunctionRegistry; }
-
+	
+	PCG_API static bool IsPCGModuleLoaded();
+	
 private:
 	FPCGGetDataFunctionRegistry GetDataFunctionRegistry;
 
diff --git a/Engine/Plugins/PCGInterops/PCGGeometryScriptInterop/Source/PCGGeometryScriptInterop/Private/Data/PCGDynamicMeshData.cpp b/Engine/Plugins/PCGInterops/PCGGeometryScriptInterop/Source/PCGGeometryScriptInterop/Private/Data/PCGDynamicMeshData.cpp
new file mode 100644
--- /dev/null
+++ b/Engine/Plugins/PCGInterops/PCGGeometryScriptInterop/Source/PCGGeometryScriptInterop/Private/Data/PCGDynamicMeshData.cpp
@@ -0,0 +1,196 @@
+// Copyright Epic Games, Inc. All Rights Reserved.
+
+#include "Data/PCGDynamicMeshData.h"
+
+#include "PCGContext.h"
+#include "PCGModule.h"
+#include "Data/PCGPointData.h"
+#include "Elements/PCGVolumeSampler.h"
+
+#include "UDynamicMesh.h"
+#include "GeometryScript/MeshQueryFunctions.h"
+#include "Helpers/PCGHelpers.h"
+#include "Serialization/ArchiveCrc32.h"
+
+UPCGDynamicMeshData::UPCGDynamicMeshData(const FObjectInitializer& ObjectInitializer)
+	: UPCGSpatialData(ObjectInitializer)
+{
+	DynamicMesh = ObjectInitializer.CreateDefaultSubobject<UDynamicMesh>(this, TEXT("DynamicMesh"));
+}
+
+void UPCGDynamicMeshData::Initialize(UDynamicMesh* InMesh, FPCGContext* Context, bool bCanTakeOwnership)
+{
+	check(DynamicMesh);
+	if (bCanTakeOwnership)
+	{
+		Initialize(std::move(InMesh->GetMeshRef()));
+	}
+	else
+	{
+		// Make a copy
+		Initialize(UE::Geometry::FDynamicMesh3(InMesh->GetMeshRef()));
+	}
+}
+
+void UPCGDynamicMeshData::Initialize(UE::Geometry::FDynamicMesh3&& InMesh, FPCGContext* Context)
+{
+	check(DynamicMesh);
+	DynamicMesh->SetMesh(std::move(InMesh));
+	ResetBounds();
+	
+	bDynamicMeshOctreeIsDirty = true;
+}
+
+void UPCGDynamicMeshData::ResetBounds() const
+{
+	check(DynamicMesh);
+	CachedBounds = static_cast<FBox>(DynamicMesh->GetMeshRef().GetBounds());
+	bDynamicMeshBoundsAreDirty = false;
+}
+
+FBox UPCGDynamicMeshData::GetBounds() const
+{
+	if (bDynamicMeshBoundsAreDirty)
+	{
+		DynamicMeshBoundsLock.Lock();
+		if (bDynamicMeshBoundsAreDirty)
+		{
+			ResetBounds();
+		}
+		DynamicMeshBoundsLock.Unlock();
+	}
+
+	return CachedBounds;
+}
+
+void UPCGDynamicMeshData::AddToCrc(FArchiveCrc32& Ar, bool bFullDataCrc) const
+{
+	Super::AddToCrc(Ar, bFullDataCrc);
+
+	if (bFullDataCrc)
+	{
+		Ar << DynamicMesh->GetMeshRef();
+	}
+	else
+	{
+		AddUIDToCrc(Ar);
+	}
+}
+
+bool UPCGDynamicMeshData::SamplePoint(const FTransform& InTransform, const FBox& InBounds, FPCGPoint& OutPoint, UPCGMetadata* OutMetadata) const
+{
+	if (!DynamicMesh)
+	{
+		return false;
+	}
+	
+	const FVector BoxCenter = InTransform.TransformPosition(InBounds.GetCenter());
+	const FVector HalfScaledExtents(InBounds.GetExtent() * InTransform.GetScale3D());
+	const UE::Geometry::FAxisAlignedBox3d ScaledBounds(BoxCenter - HalfScaledExtents, BoxCenter + HalfScaledExtents);
+	
+	if (!ScaledBounds.Intersects(GetBounds()))
+	{
+		return false;
+	}
+
+	const UE::Geometry::FDynamicMeshOctree3& Octree = GetDynamicMeshOctree();
+	const UE::Geometry::FDynamicMesh3& MeshRef = DynamicMesh->GetMeshRef();
+	
+
+	// TODO: Awful way to do this, need to implement a triangle AABB box intersection if we want more precision. But not even sure that is the right approach using the Octree
+	auto TriangleIntersectionTest = [this, &ScaledBounds, MeshRef](int32 TriangleID) -> bool
+	{
+		return ScaledBounds.Intersects(MeshRef.GetTriBounds(TriangleID));
+	};
+
+	const bool bIntersect = Octree.ParallelOverlapAnyQuery(ScaledBounds, TriangleIntersectionTest) != INDEX_NONE;
+
+	if (bIntersect)
+	{
+		OutPoint.Transform = InTransform;
+		OutPoint.SetLocalBounds(InBounds);
+		OutPoint.Density = 1.0f;
+		return true;
+	}
+	else
+	{
+		return false;
+	}
+}
+
+const UE::Geometry::FDynamicMeshOctree3& UPCGDynamicMeshData::GetDynamicMeshOctree() const
+{
+	if (DynamicMesh && bDynamicMeshOctreeIsDirty)
+	{
+		FScopeLock Lock(&DynamicMeshOctreeLock);
+
+		if (bDynamicMeshOctreeIsDirty)
+		{
+			if (DynamicMeshOctree.Mesh != DynamicMesh->GetMeshPtr())
+			{
+				// Was never initialized, or with another mesh, reset it
+				if (DynamicMeshOctree.Mesh)
+				{
+					DynamicMeshOctree = UE::Geometry::FDynamicMeshOctree3{};
+				}
+				
+				DynamicMeshOctree.Initialize(&DynamicMesh->GetMeshRef());
+			}
+			else
+			{
+				// We can re-insert all triangles to avoid recomputing the whole octree
+				DynamicMeshOctree.ReinsertTriangles(DynamicMesh->GetMeshRef().TriangleIndicesItr());
+			}
+			
+			bDynamicMeshOctreeIsDirty = false;
+		}
+	}
+
+	return DynamicMeshOctree;
+}
+
+const UPCGPointData* UPCGDynamicMeshData::ToPointData(FPCGContext* Context, const FBox& InBounds) const
+{
+	TRACE_CPUPROFILER_EVENT_SCOPE(UPCGPrimitiveData::CreatePointData);
+
+	check(DynamicMesh);
+
+	FGeometryScriptVectorList VerticesList;
+	bool Dummy;
+	UGeometryScriptLibrary_MeshQueryFunctions::GetAllVertexPositions(DynamicMesh, VerticesList, /*bSkipGaps=*/false, /*bHasVertexIDGaps=*/Dummy);
+
+	TArray<FVector>& Vertices = *VerticesList.List.Get();
+
+	if (Vertices.IsEmpty())
+	{
+		return nullptr;
+	}
+	
+	UPCGPointData* Data = FPCGContext::NewObject_AnyThread<UPCGPointData>(Context);
+	TArray<FPCGPoint>& Points = Data->GetMutablePoints();
+	Points.Reserve(Vertices.Num());
+	for (const FVector& Vertex : Vertices)
+	{
+		if (InBounds.IsValid && !InBounds.IsInside(Vertex))
+		{
+			continue;
+		}
+		
+		FPCGPoint& Point = Points.Emplace_GetRef();
+		Point.Transform.SetLocation(Vertex);
+		PCGHelpers::ComputeSeedFromPosition(Vertex);
+	}
+	
+	UE_LOG(LogPCG, Verbose, TEXT("Dynamic Mesh extracted %d points"), Points.Num());
+
+	return Data;
+}
+
+UPCGSpatialData* UPCGDynamicMeshData::CopyInternal(FPCGContext* Context) const
+{
+	UPCGDynamicMeshData* NewDynamicMeshData = FPCGContext::NewObject_AnyThread<UPCGDynamicMeshData>(Context);
+	
+	NewDynamicMeshData->Initialize(DynamicMesh, Context, /*bCanTakeOwnership=*/false);
+
+	return NewDynamicMeshData;
+}
diff --git a/Engine/Plugins/PCGInterops/PCGGeometryScriptInterop/Source/PCGGeometryScriptInterop/Private/Editor/PCGDynamicMeshDataVisualization.cpp b/Engine/Plugins/PCGInterops/PCGGeometryScriptInterop/Source/PCGGeometryScriptInterop/Private/Editor/PCGDynamicMeshDataVisualization.cpp
new file mode 100644
--- /dev/null
+++ b/Engine/Plugins/PCGInterops/PCGGeometryScriptInterop/Source/PCGGeometryScriptInterop/Private/Editor/PCGDynamicMeshDataVisualization.cpp
@@ -0,0 +1,36 @@
+// Copyright Epic Games, Inc. All Rights Reserved.
+
+#include "Editor/PCGDynamicMeshDataVisualization.h"
+
+#if WITH_EDITOR
+
+#include "Data/PCGDynamicMeshData.h"
+#include "Helpers/PCGHelpers.h"
+#include "Resources/PCGDynamicMeshManagedComponent.h"
+
+#include "UDynamicMesh.h"
+#include "Components/DynamicMeshComponent.h"
+
+#define LOCTEXT_NAMESPACE "PCGDynamicMeshDataVisualization"
+
+void FPCGDynamicMeshDataVisualization::ExecuteDebugDisplay(FPCGContext* Context, const UPCGSettingsInterface* SettingsInterface, const UPCGData* Data, AActor* TargetActor) const
+{
+	const UPCGDynamicMeshData* DynMeshData = CastChecked<UPCGDynamicMeshData>(Data);
+	if (!DynMeshData)
+	{
+		return;
+	}
+
+	// We force debug resources to be transient.
+	UPCGDynamicMeshManagedComponent* ManagedComponent = PCGDynamicMeshManagedComponent::GetOrCreateDynamicMeshManagedComponent(Context, SettingsInterface, DynMeshData, TargetActor, EPCGEditorDirtyMode::Preview);
+	
+	if (ManagedComponent && ManagedComponent->GetComponent())
+	{
+		// Needs to be a copy.
+		ManagedComponent->GetComponent()->SetMesh(UE::Geometry::FDynamicMesh3(DynMeshData->GetDynamicMesh()->GetMeshRef()));
+	}
+}
+
+#undef LOCTEXT_NAMESPACE
+
+#endif // WITH_EDITOR
\ No newline at end of file
diff --git a/Engine/Plugins/PCGInterops/PCGGeometryScriptInterop/Source/PCGGeometryScriptInterop/Private/Editor/PCGDynamicMeshDataVisualization.h b/Engine/Plugins/PCGInterops/PCGGeometryScriptInterop/Source/PCGGeometryScriptInterop/Private/Editor/PCGDynamicMeshDataVisualization.h
new file mode 100644
--- /dev/null
+++ b/Engine/Plugins/PCGInterops/PCGGeometryScriptInterop/Source/PCGGeometryScriptInterop/Private/Editor/PCGDynamicMeshDataVisualization.h
@@ -0,0 +1,22 @@
+// Copyright Epic Games, Inc. All Rights Reserved.
+
+#pragma once
+
+#if WITH_EDITOR
+#include "PCGDataVisualization.h"
+#include "DataVisualizations/PCGSpatialDataVisualization.h"
+
+class AActor;
+class UPCGData;
+class UPCGPointData;
+struct FPCGContext;
+
+class FPCGDynamicMeshDataVisualization : public IPCGSpatialDataVisualization
+{
+public:
+	// ~Begin IPCGDataVisualization interface
+	virtual void ExecuteDebugDisplay(FPCGContext* Context, const UPCGSettingsInterface* SettingsInterface, const UPCGData* Data, AActor* TargetActor) const override;
+	// ~End IPCGDataVisualization interface
+};
+
+#endif // WITH_EDITOR
\ No newline at end of file
diff --git a/Engine/Plugins/PCGInterops/PCGGeometryScriptInterop/Source/PCGGeometryScriptInterop/Private/Elements/PCGGetDynamicMeshData.cpp b/Engine/Plugins/PCGInterops/PCGGeometryScriptInterop/Source/PCGGeometryScriptInterop/Private/Elements/PCGGetDynamicMeshData.cpp
new file mode 100644
--- /dev/null
+++ b/Engine/Plugins/PCGInterops/PCGGeometryScriptInterop/Source/PCGGeometryScriptInterop/Private/Elements/PCGGetDynamicMeshData.cpp
@@ -0,0 +1,130 @@
+// Copyright Epic Games, Inc. All Rights Reserved.
+
+#include "Elements/PCGGetDynamicMeshData.h"
+
+#include "PCGContext.h"
+#include "Data/PCGDynamicMeshData.h"
+#include "Data/PCGGetDataFunctionRegistry.h"
+#include "Helpers/PCGHelpers.h"
+
+#include "DynamicMeshActor.h"
+
+#define LOCTEXT_NAMESPACE "PCGGetDynamicMeshDataElement"
+
+#if WITH_EDITOR
+FName UPCGGetDynamicMeshDataSettings::GetDefaultNodeName() const
+{
+	return FName(TEXT("GetDynamicMeshData"));
+}
+
+FText UPCGGetDynamicMeshDataSettings::GetDefaultNodeTitle() const
+{
+	return LOCTEXT("NodeTitle", "Get Dynamic Mesh Data");
+}
+#endif // WITH_EDITOR
+
+TArray<FPCGPinProperties> UPCGGetDynamicMeshDataSettings::OutputPinProperties() const
+{ 
+	TArray<FPCGPinProperties> PinProperties;
+	PinProperties.Emplace(PCGPinConstants::DefaultOutputLabel, EPCGDataType::DynamicMesh);
+
+	return PinProperties;
+}
+
+bool PCGGetDynamicMeshData::GetDynamicMeshDataFromActor(FPCGContext* InContext, const FPCGGetDataFunctionRegistryParams& InParams, AActor* InActor, FPCGGetDataFunctionRegistryOutput& Output)
+{
+	// Request Dynamic Mesh Data filter explicitly, otherwise it would consume too many actors/components.
+	if (InParams.DataTypeFilter != EPCGDataType::DynamicMesh)
+	{
+		return false;
+	}
+	
+	check(InActor);
+
+	// Early out if the actor gets rejected by the component selector
+	if (InParams.ComponentSelector && !InParams.ComponentSelector->FilterActor(InActor))
+	{
+		return false;
+	}
+	
+	if (ADynamicMeshActor* DynMeshActor = Cast<ADynamicMeshActor>(InActor))
+	{
+		if (!DynMeshActor->GetDynamicMeshComponent() || !DynMeshActor->GetDynamicMeshComponent()->GetDynamicMesh())
+		{
+			return false;
+		}
+		
+		auto NameTagsToStringTags = [](const FName& InName) { return InName.ToString(); };
+		TSet<FString> ActorTags;
+		Algo::Transform(InActor->Tags, ActorTags, NameTagsToStringTags);
+		
+		UPCGDynamicMeshData* Data = FPCGContext::NewObject_AnyThread<UPCGDynamicMeshData>(InContext);
+		Data->Initialize(DynMeshActor->GetDynamicMeshComponent()->GetDynamicMesh(), InContext, /*bCanTakeOwnership=*/ false);
+
+		FPCGTaggedData& TaggedData = Output.Collection.TaggedData.Emplace_GetRef();
+		TaggedData.Data = Data;
+		TaggedData.Tags = ActorTags;
+		return true;
+	}
+	else
+	{
+		return false;
+	}
+}
+
+bool PCGGetDynamicMeshData::GetDynamicMeshDataFromComponent(FPCGContext* InContext, const FPCGGetDataFunctionRegistryParams& InParams, UActorComponent* InActorComponent, FPCGGetDataFunctionRegistryOutput& Output)
+{
+	check(InActorComponent);
+
+	// Request Dynamic Mesh Data filter explicitly, otherwise it would consume too many actors/components.
+	if (InParams.DataTypeFilter != EPCGDataType::DynamicMesh)
+	{
+		return false;
+	}
+
+	if (InParams.bIgnorePCGGeneratedComponents && InActorComponent->ComponentTags.Contains(PCGHelpers::DefaultPCGTag))
+	{
+		return false;
+	}
+	
+	// Check if it is a scene component and try to extract it.
+	if (USceneComponent* SceneComponent = Cast<USceneComponent>(InActorComponent))
+	{
+		UPCGDynamicMeshData* OutputData = FPCGContext::NewObject_AnyThread<UPCGDynamicMeshData>(InContext);
+		UDynamicMesh* NewDynamicMesh = OutputData->GetMutableDynamicMesh();
+		check(NewDynamicMesh);
+
+		const UPCGGetDynamicMeshDataSettings* Settings = InContext ? InContext->GetInputSettings<UPCGGetDynamicMeshDataSettings>() : nullptr;
+		FGeometryScriptCopyMeshFromComponentOptions Options = Settings ? Settings->Options : FGeometryScriptCopyMeshFromComponentOptions{};
+
+		// TODO: Might want to store the transform on the data?
+		FTransform Transform;
+		EGeometryScriptOutcomePins Outcome;
+		UGeometryScriptLibrary_SceneUtilityFunctions::CopyMeshFromComponent(SceneComponent, NewDynamicMesh, Options, /*bTransformToWorld=*/ false, Transform, Outcome);
+
+		if (Outcome == EGeometryScriptOutcomePins::Success)
+		{
+			auto NameTagsToStringTags = [](const FName& InName) { return InName.ToString(); };
+
+			FPCGTaggedData& TaggedData = Output.Collection.TaggedData.Emplace_GetRef();
+			TaggedData.Data = OutputData;
+			
+			Algo::Transform(SceneComponent->ComponentTags, TaggedData.Tags, NameTagsToStringTags);
+
+			if (InParams.bAddActorTags && SceneComponent->GetOwner())
+			{
+				TSet<FString> ActorTags;
+				Algo::Transform(SceneComponent->GetOwner()->Tags, ActorTags, NameTagsToStringTags);
+				TaggedData.Tags.Append(ActorTags);
+			}
+		}
+
+		return true;
+	}
+	else
+	{
+		return false;
+	}
+}
+
+#undef LOCTEXT_NAMESPACE
diff --git a/Engine/Plugins/PCGInterops/PCGGeometryScriptInterop/Source/PCGGeometryScriptInterop/Private/PCGGeometryScriptInteropModule.cpp b/Engine/Plugins/PCGInterops/PCGGeometryScriptInterop/Source/PCGGeometryScriptInterop/Private/PCGGeometryScriptInteropModule.cpp
--- a/Engine/Plugins/PCGInterops/PCGGeometryScriptInterop/Source/PCGGeometryScriptInterop/Private/PCGGeometryScriptInteropModule.cpp
+++ b/Engine/Plugins/PCGInterops/PCGGeometryScriptInterop/Source/PCGGeometryScriptInterop/Private/PCGGeometryScriptInteropModule.cpp
@@ -2,6 +2,14 @@
 
 #include "Modules/ModuleManager.h"
 
+#include "PCGModule.h"
+#include "Data/PCGDynamicMeshData.h"
+#include "Elements/PCGGetDynamicMeshData.h"
+
+#if WITH_EDITOR
+#include "Editor/PCGDynamicMeshDataVisualization.h"
+#endif // WITH_EDITOR
+
 class FPCGGeometryScriptInteropModule final : public IModuleInterface
 {
 public:
@@ -10,7 +18,39 @@ public:
 	{
 		return true;
 	}
+	
+	virtual void StartupModule() override
+	{
+#if WITH_EDITOR
+		FPCGDataVisualizationRegistry& DataVisRegistry = FPCGModule::GetMutablePCGDataVisualizationRegistry();
+		DataVisRegistry.RegisterPCGDataVisualization(UPCGDynamicMeshData::StaticClass(), MakeUnique<const FPCGDynamicMeshDataVisualization>());
+#endif // WITH_EDITOR
+
+		FPCGGetDataFunctionRegistry& PCGDataFunctionRegistry = FPCGModule::MutableGetDataFunctionRegistry();
+		GetActorDataFunctionHandle = PCGDataFunctionRegistry.RegisterDataFromActorFunction(&PCGGetDynamicMeshData::GetDynamicMeshDataFromActor);
+		GetComponentDataFunctionHandle = PCGDataFunctionRegistry.RegisterDataFromComponentFunction(&PCGGetDynamicMeshData::GetDynamicMeshDataFromComponent);
+	}
+
+	virtual void ShutdownModule() override
+	{
+		// No need to unregister if the PCG module is already dead.
+		if (FPCGModule::IsPCGModuleLoaded())
+		{
+			FPCGGetDataFunctionRegistry& PCGDataFunctionRegistry = FPCGModule::MutableGetDataFunctionRegistry();
+			PCGDataFunctionRegistry.UnregisterDataFromActorFunction(GetActorDataFunctionHandle);
+			PCGDataFunctionRegistry.UnregisterDataFromComponentFunction(GetComponentDataFunctionHandle);
+		
+#if WITH_EDITOR
+			FPCGDataVisualizationRegistry& DataVisRegistry = FPCGModule::GetMutablePCGDataVisualizationRegistry();
+			DataVisRegistry.UnregisterPCGDataVisualization(UPCGDynamicMeshData::StaticClass());
+#endif // WITH_EDITOR
+		}
+	}
 	//~ End IModuleInterface implementation
+	
+private:
+	FPCGGetDataFunctionRegistry::FFunctionHandle GetActorDataFunctionHandle = (uint64)(-1);
+	FPCGGetDataFunctionRegistry::FFunctionHandle GetComponentDataFunctionHandle = (uint64)(-1);
 };
 
 IMPLEMENT_MODULE(FPCGGeometryScriptInteropModule, PCGGeometryScriptInterop);
diff --git a/Engine/Plugins/PCGInterops/PCGGeometryScriptInterop/Source/PCGGeometryScriptInterop/Private/Resources/PCGDynamicMeshManagedComponent.cpp b/Engine/Plugins/PCGInterops/PCGGeometryScriptInterop/Source/PCGGeometryScriptInterop/Private/Resources/PCGDynamicMeshManagedComponent.cpp
new file mode 100644
--- /dev/null
+++ b/Engine/Plugins/PCGInterops/PCGGeometryScriptInterop/Source/PCGGeometryScriptInterop/Private/Resources/PCGDynamicMeshManagedComponent.cpp
@@ -0,0 +1,147 @@
+// Copyright Epic Games, Inc. All Rights Reserved.
+
+#include "Resources/PCGDynamicMeshManagedComponent.h"
+
+#include "PCGComponent.h"
+#include "PCGContext.h"
+#include "Data/PCGDynamicMeshData.h"
+#include "Helpers/PCGHelpers.h"
+
+#include "Components/DynamicMeshComponent.h"
+
+#define LOCTEXT_NAMESPACE "PCGDynamicMeshManagedComponent"
+
+void UPCGDynamicMeshManagedComponent::ForgetComponent()
+{
+	Super::ForgetComponent();
+	CachedRawComponentPtr = nullptr;
+}
+
+void UPCGDynamicMeshManagedComponent::MarkAsReused()
+{
+	Super::MarkAsReused();
+	
+	// We need to reset the transform is we re-use the component. Similar to ISMC code.
+	if (UDynamicMeshComponent* Component = GetComponent())
+	{
+		FVector TentativeRootLocation = FVector::ZeroVector;
+		
+		if (USceneComponent* RootComponent = Component->GetAttachmentRoot())
+		{
+			TentativeRootLocation = RootComponent->GetComponentLocation();
+		}
+
+		// Since this is technically 'moving' the component, we need to unregister it before moving otherwise we could get a warning that we're moving a component with static mobility
+		Component->UnregisterComponent();
+		Component->SetWorldTransform(FTransform(FQuat::Identity, TentativeRootLocation, FVector::OneVector));
+		Component->RegisterComponent();
+	}
+}
+
+UDynamicMeshComponent* UPCGDynamicMeshManagedComponent::GetComponent() const
+{
+	if (!CachedRawComponentPtr)
+	{
+		UDynamicMeshComponent* GeneratedComponentPtr = Cast<UDynamicMeshComponent>(GeneratedComponent.Get());
+
+		// Implementation note:
+		// There is no surefire way to make sure that we can use the raw pointer UNLESS it is from the same owner
+		if (GeneratedComponentPtr && Cast<UPCGComponent>(GetOuter()) && GeneratedComponentPtr->GetOwner() == Cast<UPCGComponent>(GetOuter())->GetOwner())
+		{
+			CachedRawComponentPtr = GeneratedComponentPtr;
+		}
+
+		return GeneratedComponentPtr;
+	}
+
+	return CachedRawComponentPtr;
+}
+
+void UPCGDynamicMeshManagedComponent::SetComponent(UDynamicMeshComponent* InComponent)
+{
+	GeneratedComponent = InComponent;
+	CachedRawComponentPtr = InComponent;
+}
+
+UPCGDynamicMeshManagedComponent* PCGDynamicMeshManagedComponent::GetOrCreateDynamicMeshManagedComponent(FPCGContext* Context, const UPCGSettingsInterface* SettingsInterface, const UPCGDynamicMeshData* InMeshData, AActor* TargetActor, TOptional<EPCGEditorDirtyMode> OptionalDirtyModeOverride)
+{
+	check(InMeshData);
+	
+	if (!TargetActor)
+	{
+		PCGLog::LogErrorOnGraph(LOCTEXT("NoTargetActor", "Cannot execute debug display for Dynamic Mesh data with no target actor."), Context);
+		return nullptr;
+	}
+
+	UPCGComponent* SourceComponent = Context->SourceComponent.Get();
+	if (!SourceComponent || !SettingsInterface)
+	{
+		return nullptr;
+	}
+	
+	const FPCGCrc CRC = InMeshData->GetOrComputeCrc(/*bFullDataCrc=*/true);
+
+	UPCGDynamicMeshManagedComponent* ExistingResource = nullptr;
+	SourceComponent->ForEachManagedResource([&ExistingResource, &CRC](UPCGManagedResource* Resource)
+	{
+		// If we already found a valid resource, just skip until the end.
+		if (ExistingResource)
+		{
+			return;
+		}
+		
+		UPCGDynamicMeshManagedComponent* DynMeshResource = Cast<UPCGDynamicMeshManagedComponent>(Resource);
+		if (!DynMeshResource)
+		{
+			return;
+		}
+		
+		if (DynMeshResource->GetDataUID() == CRC.GetValue() && DynMeshResource->CanBeUsed())
+		{
+			ExistingResource = DynMeshResource;
+		}
+	});
+
+	if (!ExistingResource)
+	{
+		ExistingResource = FPCGContext::NewObject_AnyThread<UPCGDynamicMeshManagedComponent>(Context, SourceComponent);
+		SourceComponent->AddToManagedResources(ExistingResource);
+		
+		ExistingResource->SetDataUID(CRC.GetValue());
+	}
+
+	check(ExistingResource);
+	
+	ExistingResource->MarkAsUsed();
+	const EPCGEditorDirtyMode DirtyMode = OptionalDirtyModeOverride.IsSet() ? OptionalDirtyModeOverride.GetValue() : SourceComponent->GetEditingMode();
+	
+#if WITH_EDITOR
+	ExistingResource->ChangeTransientState(DirtyMode);
+#endif // WITH_EDITOR
+
+	UDynamicMeshComponent* DynMeshComponent = ExistingResource->GetComponent();
+	if (!DynMeshComponent)
+	{
+		DynMeshComponent = FPCGContext::NewObject_AnyThread<UDynamicMeshComponent>(Context, TargetActor);
+		if (DirtyMode == EPCGEditorDirtyMode::Preview)
+		{
+			DynMeshComponent->SetFlags(EObjectFlags::RF_Transient);
+		}
+
+		DynMeshComponent->RegisterComponent();
+		TargetActor->AddInstanceComponent(DynMeshComponent);
+
+		// Mimicking Static mesh managed resources
+		DynMeshComponent->AttachToComponent(TargetActor->GetRootComponent(), FAttachmentTransformRules(EAttachmentRule::KeepRelative, EAttachmentRule::KeepWorld, EAttachmentRule::KeepWorld, false));
+		
+		ExistingResource->SetComponent(DynMeshComponent);
+	}
+
+	// Add default tags. It's the callee responsibility to re-apply other tags.
+	DynMeshComponent->ComponentTags.AddUnique(SourceComponent->GetFName());
+	DynMeshComponent->ComponentTags.AddUnique(PCGHelpers::DefaultPCGTag);
+	
+	return ExistingResource;
+}
+
+#undef LOCTEXT_NAMESPACE
\ No newline at end of file
diff --git a/Engine/Plugins/PCGInterops/PCGGeometryScriptInterop/Source/PCGGeometryScriptInterop/Public/Data/PCGDynamicMeshData.h b/Engine/Plugins/PCGInterops/PCGGeometryScriptInterop/Source/PCGGeometryScriptInterop/Public/Data/PCGDynamicMeshData.h
new file mode 100644
--- /dev/null
+++ b/Engine/Plugins/PCGInterops/PCGGeometryScriptInterop/Source/PCGGeometryScriptInterop/Public/Data/PCGDynamicMeshData.h
@@ -0,0 +1,71 @@
+// Copyright Epic Games, Inc. All Rights Reserved.
+
+#pragma once
+
+#include "Data/PCGSpatialData.h"
+
+#include "DynamicMesh/DynamicMesh3.h"
+#include "DynamicMesh/DynamicMeshOctree3.h"
+
+#include "Misc/SpinLock.h"
+
+#include "PCGDynamicMeshData.generated.h"
+
+struct FPCGContext;
+class UDynamicMesh;
+
+UCLASS(BlueprintType, ClassGroup = (Procedural))
+class PCGGEOMETRYSCRIPTINTEROP_API UPCGDynamicMeshData : public UPCGSpatialData
+{
+	GENERATED_BODY()
+
+public:
+	UPCGDynamicMeshData(const FObjectInitializer& ObjectInitializer);
+	
+	void Initialize(UDynamicMesh* InMesh, FPCGContext* Context = nullptr, bool bCanTakeOwnership = false);
+	void Initialize(UE::Geometry::FDynamicMesh3&& InMesh, FPCGContext* Context = nullptr);
+	
+	UFUNCTION(BlueprintCallable, Category="DynamicMesh", meta = (DisplayName = "Initialize"))
+	void K2_Initialize(UDynamicMesh* InMesh, FPCGContext& Context) { Initialize(InMesh, &Context); }
+	
+	// ~Begin UPCGData interface
+	virtual EPCGDataType GetDataType() const override { return EPCGDataType::DynamicMesh; }
+	virtual void AddToCrc(FArchiveCrc32& Ar, bool bFullDataCrc) const override;
+	// ~End UPCGData interface
+
+	// ~Begin UPCGSpatialData interface
+	virtual int GetDimension() const override { return 3; }
+	virtual FBox GetBounds() const override;
+	virtual bool SamplePoint(const FTransform& Transform, const FBox& Bounds, FPCGPoint& OutPoint, UPCGMetadata* OutMetadata) const override;
+	// TODO needs an implementation to support projection
+	//virtual bool ProjectPoint(const FTransform& InTransform, const FBox& InBounds, const FPCGProjectionParams& InParams, FPCGPoint& OutPoint, UPCGMetadata* OutMetadata) const;
+	//~End UPCGSpatialData interface
+
+	const UE::Geometry::FDynamicMeshOctree3& GetDynamicMeshOctree() const;
+	
+	UDynamicMesh* GetMutableDynamicMesh() { bDynamicMeshBoundsAreDirty = true; bDynamicMeshOctreeIsDirty = true; return DynamicMesh; }
+	const UDynamicMesh* GetDynamicMesh() const { return DynamicMesh; }
+
+protected:
+	// ~Begin UPCGSpatialData interface
+	virtual UPCGSpatialData* CopyInternal(FPCGContext* Context) const override;
+public:
+	virtual const UPCGPointData* ToPointData(FPCGContext* Context, const FBox& InBounds) const override;
+	//~End UPCGSpatialData interface
+
+private:
+	// const but will set the mutable CachedBounds
+	void ResetBounds() const;
+
+protected:
+	UPROPERTY(BlueprintReadOnly, VisibleAnywhere, Instanced, Category="DynamicMesh")
+	TObjectPtr<UDynamicMesh> DynamicMesh;
+
+	mutable UE::Geometry::FDynamicMeshOctree3 DynamicMeshOctree;
+	mutable bool bDynamicMeshOctreeIsDirty = true;
+	mutable FCriticalSection DynamicMeshOctreeLock;
+	
+	mutable FBox CachedBounds = FBox(EForceInit::ForceInit);
+	mutable bool bDynamicMeshBoundsAreDirty = true;
+	mutable UE::FSpinLock DynamicMeshBoundsLock;
+};
diff --git a/Engine/Plugins/PCGInterops/PCGGeometryScriptInterop/Source/PCGGeometryScriptInterop/Public/Elements/PCGGetDynamicMeshData.h b/Engine/Plugins/PCGInterops/PCGGeometryScriptInterop/Source/PCGGeometryScriptInterop/Public/Elements/PCGGetDynamicMeshData.h
new file mode 100644
--- /dev/null
+++ b/Engine/Plugins/PCGInterops/PCGGeometryScriptInterop/Source/PCGGeometryScriptInterop/Public/Elements/PCGGetDynamicMeshData.h
@@ -0,0 +1,50 @@
+// Copyright Epic Games, Inc. All Rights Reserved.
+
+#pragma once
+
+#include "Elements/PCGDataFromActor.h"
+
+#include "GeometryScript/SceneUtilityFunctions.h"
+
+#include "PCGGetDynamicMeshData.generated.h"
+
+struct FPCGGetDataFunctionRegistryOutput;
+struct FPCGGetDataFunctionRegistryParams;
+
+UCLASS(MinimalAPI, BlueprintType, ClassGroup = (Procedural))
+class UPCGGetDynamicMeshDataSettings : public UPCGDataFromActorSettings
+{
+	GENERATED_BODY()
+
+public:
+	//~Begin UPCGSettings interface
+#if WITH_EDITOR
+	virtual FName GetDefaultNodeName() const override;
+	virtual FText GetDefaultNodeTitle() const override;
+#endif
+
+protected:
+	virtual TArray<FPCGPinProperties> OutputPinProperties() const override;
+	//~End UPCGSettings interface
+	
+	//~Begin UPCGDataFromActorSettings interface
+public:
+	virtual EPCGDataType GetDataFilter() const override { return EPCGDataType::DynamicMesh; }
+
+protected:
+#if WITH_EDITOR
+	virtual bool DisplayModeSettings() const override { return false; }
+#endif
+	//~End UPCGDataFromActorSettings
+	
+public:
+	/** If data is coming from a component, you can impact the options there. */
+	UPROPERTY(BlueprintReadWrite, EditAnywhere, Category="Dynamic Mesh Settings", meta = (PCG_Overridable, ShowOnlyInnerProperties))
+	FGeometryScriptCopyMeshFromComponentOptions Options;
+};
+
+namespace PCGGetDynamicMeshData
+{
+	PCGGEOMETRYSCRIPTINTEROP_API bool GetDynamicMeshDataFromActor(FPCGContext*, const FPCGGetDataFunctionRegistryParams&, AActor*, FPCGGetDataFunctionRegistryOutput&);
+	PCGGEOMETRYSCRIPTINTEROP_API bool GetDynamicMeshDataFromComponent(FPCGContext*, const FPCGGetDataFunctionRegistryParams&, UActorComponent*, FPCGGetDataFunctionRegistryOutput&);
+}
\ No newline at end of file
diff --git a/Engine/Plugins/PCGInterops/PCGGeometryScriptInterop/Source/PCGGeometryScriptInterop/Public/Resources/PCGDynamicMeshManagedComponent.h b/Engine/Plugins/PCGInterops/PCGGeometryScriptInterop/Source/PCGGeometryScriptInterop/Public/Resources/PCGDynamicMeshManagedComponent.h
new file mode 100644
--- /dev/null
+++ b/Engine/Plugins/PCGInterops/PCGGeometryScriptInterop/Source/PCGGeometryScriptInterop/Public/Resources/PCGDynamicMeshManagedComponent.h
@@ -0,0 +1,47 @@
+// Copyright Epic Games, Inc. All Rights Reserved.
+
+#pragma once
+
+#include "PCGCommon.h"
+#include "PCGManagedResource.h"
+
+#include "PCGDynamicMeshManagedComponent.generated.h"
+
+class AActor;
+struct FPCGContext;
+class UDynamicMeshComponent;
+class UPCGData;
+class UPCGDynamicMeshData;
+class UPCGSettingsInterface;
+
+UCLASS(BlueprintType, MinimalAPI)
+class UPCGDynamicMeshManagedComponent : public UPCGManagedComponent
+{
+	GENERATED_BODY()
+
+public:
+	//~Begin UPCGManagedComponents interface
+	virtual void ResetComponent() override { /* Does nothing, but implementation is required to support reuse. */ }
+	virtual bool SupportsComponentReset() const override { return true; }
+	virtual void MarkAsReused() override;
+	virtual void ForgetComponent() override;
+	//~End UPCGManagedComponents interface
+
+	PCGGEOMETRYSCRIPTINTEROP_API UDynamicMeshComponent* GetComponent() const;
+	PCGGEOMETRYSCRIPTINTEROP_API void SetComponent(UDynamicMeshComponent* InComponent);
+
+	uint64 GetDataUID() const { return DataUID; }
+	void SetDataUID(uint64 InDataUID) { DataUID = InDataUID; }
+
+protected:
+	UPROPERTY(Transient)
+	uint64 DataUID = -1;
+
+	// Cached raw pointer to UDynamicMeshComponent
+	mutable UDynamicMeshComponent* CachedRawComponentPtr = nullptr;
+};
+
+namespace PCGDynamicMeshManagedComponent
+{
+	PCGGEOMETRYSCRIPTINTEROP_API UPCGDynamicMeshManagedComponent* GetOrCreateDynamicMeshManagedComponent(FPCGContext* Context, const UPCGSettingsInterface* SettingsInterface, const UPCGDynamicMeshData* InMeshData, AActor* TargetActor, TOptional<EPCGEditorDirtyMode> OptionalDirtyModeOverride = TOptional<EPCGEditorDirtyMode>{});
+}
\ No newline at end of file
