[PCG] Update Mesh sampler to support Dynamic Meshes as input

#rnx

diff --git a/Engine/Plugins/PCGInterops/PCGGeometryScriptInterop/Source/PCGGeometryScriptInterop/Private/Elements/PCGMeshSampler.cpp b/Engine/Plugins/PCGInterops/PCGGeometryScriptInterop/Source/PCGGeometryScriptInterop/Private/Elements/PCGMeshSampler.cpp
--- a/Engine/Plugins/PCGInterops/PCGGeometryScriptInterop/Source/PCGGeometryScriptInterop/Private/Elements/PCGMeshSampler.cpp
+++ b/Engine/Plugins/PCGInterops/PCGGeometryScriptInterop/Source/PCGGeometryScriptInterop/Private/Elements/PCGMeshSampler.cpp
@@ -16,6 +16,7 @@
 
 #include "UDynamicMesh.h"
 #include "ConversionUtils/SceneComponentToDynamicMesh.h"
+#include "Data/PCGDynamicMeshData.h"
 #include "Engine/AssetManager.h"
 #include "Engine/StaticMesh.h"
 #include "GameFramework/Actor.h"
@@ -503,6 +504,24 @@ bool FPCGMeshSamplerElement::PrepareDataInternal(FPCGContext* InContext) const
 		}
 	}
 
+	// Manually adding incoming Dynamic meshes data input, copy is necessary if we voxelize, otherwise we'll const cast as it won't be modified (but GeometryScript API is not const friendly).
+	for (const FPCGTaggedData& Input : InContext->InputData.GetInputsByPin(PCGPinConstants::DefaultInputLabel))
+	{
+		if (const UPCGDynamicMeshData* InputData = Cast<const UPCGDynamicMeshData>(Input.Data))
+		{
+			if (Settings->bVoxelize)
+			{
+				UDynamicMesh* DynMesh = Context->DynamicMeshes.Add_GetRef(FPCGContext::NewObject_AnyThread<UDynamicMesh>(Context));
+				UE::Geometry::FDynamicMesh3 MeshCopy = InputData->GetDynamicMesh()->GetMeshRef();
+				DynMesh->SetMesh(std::move(MeshCopy));
+			}
+			else
+			{
+				Context->DynamicMeshes.Emplace(const_cast<UDynamicMesh*>(InputData->GetDynamicMesh()));
+			}
+		}
+	}
+
 	Debug->MarkAsGarbage();
 
 	// Reserve arrays. Add one more entry for the starting indices to have the total number of items to process.
