[PCG] First wave of Geometry Script nodes in PCG

#rnx

diff --git a/Engine/Plugins/PCG/Source/PCG/Private/Graph/PCGGraphCompiler.cpp b/Engine/Plugins/PCG/Source/PCG/Private/Graph/PCGGraphCompiler.cpp
--- a/Engine/Plugins/PCG/Source/PCG/Private/Graph/PCGGraphCompiler.cpp
+++ b/Engine/Plugins/PCG/Source/PCG/Private/Graph/PCGGraphCompiler.cpp
@@ -1190,6 +1190,7 @@ void FPCGGraphCompiler::CompileTopGraph(UPCGGraph* InGraph, uint32 GenerationGri
 			}
 			else
 			{
+				TaskInput.bIsUsedMultipleTimes = false;
 				SeenInputs.Emplace(std::move(InputOrigin));
 			}
 		}
diff --git a/Engine/Plugins/PCGInterops/PCGGeometryScriptInterop/Source/PCGGeometryScriptInterop/Private/Elements/PCGAppendMeshesFromPoints.cpp b/Engine/Plugins/PCGInterops/PCGGeometryScriptInterop/Source/PCGGeometryScriptInterop/Private/Elements/PCGAppendMeshesFromPoints.cpp
new file mode 100644
--- /dev/null
+++ b/Engine/Plugins/PCGInterops/PCGGeometryScriptInterop/Source/PCGGeometryScriptInterop/Private/Elements/PCGAppendMeshesFromPoints.cpp
@@ -0,0 +1,324 @@
+// Copyright Epic Games, Inc. All Rights Reserved.
+
+#include "Elements/PCGAppendMeshesFromPoints.h"
+
+#include "PCGContext.h"
+#include "Data/PCGDynamicMeshData.h"
+#include "Data/PCGPointData.h"
+#include "Elements/Metadata/PCGMetadataElementCommon.h"
+#include "Metadata/Accessors/PCGAttributeAccessorHelpers.h"
+
+#include "DynamicMeshEditor.h"
+#include "UDynamicMesh.h"
+#include "ConversionUtils/SceneComponentToDynamicMesh.h"
+#include "DynamicMesh/MeshIndexMappings.h"
+#include "Engine/StaticMesh.h"
+
+#define LOCTEXT_NAMESPACE "PCGAppendMeshesFromPointsElement"
+
+namespace PCGAppendMeshesFromPoints
+{
+	static const FName InDynMeshPinLabel = TEXT("InDynMesh");
+	static const FName InPointsPinLabel = TEXT("InPoints");
+	static const FName InAppendMeshPinLabel = TEXT("AppendDynMesh");
+}
+
+#if WITH_EDITOR
+FName UPCGAppendMeshesFromPointsSettings::GetDefaultNodeName() const
+{
+	return FName(TEXT("AppendMeshesFromPoints"));
+}
+
+FText UPCGAppendMeshesFromPointsSettings::GetDefaultNodeTitle() const
+{
+	return LOCTEXT("NodeTitle", "Append Meshes From Points");
+}
+
+FText UPCGAppendMeshesFromPointsSettings::GetNodeTooltipText() const
+{
+	return LOCTEXT("NodeTooltip", "Append meshes at the points transforms. Mesh can be a single static mesh, multiple meshes coming from the points or another dynamic mesh.");
+}
+
+EPCGChangeType UPCGAppendMeshesFromPointsSettings::GetChangeTypeForProperty(const FName& InPropertyName) const
+{
+	EPCGChangeType ChangeType = Super::GetChangeTypeForProperty(InPropertyName) | EPCGChangeType::Cosmetic;
+
+	if (InPropertyName == GET_MEMBER_NAME_CHECKED(UPCGAppendMeshesFromPointsSettings, Mode))
+	{
+		// If we change from/to DynamicMesh, this needs to trigger a graph recompilation.
+		ChangeType |= EPCGChangeType::Structural;
+	}
+
+	return ChangeType;
+}
+#endif // WITH_EDITOR
+
+FPCGElementPtr UPCGAppendMeshesFromPointsSettings::CreateElement() const
+{
+	return MakeShared<FPCGAppendMeshesFromPointsElement>();
+}
+
+TArray<FPCGPinProperties> UPCGAppendMeshesFromPointsSettings::InputPinProperties() const
+{
+	TArray<FPCGPinProperties> Properties;
+	Properties.Emplace_GetRef(PCGAppendMeshesFromPoints::InDynMeshPinLabel, EPCGDataType::DynamicMesh, false, false).SetRequiredPin();
+	Properties.Emplace_GetRef(PCGAppendMeshesFromPoints::InPointsPinLabel, EPCGDataType::Point, false, false).SetRequiredPin();
+
+	if (Mode == EPCGAppendMeshesFromPointsMode::DynamicMesh)
+	{
+		Properties.Emplace_GetRef(PCGAppendMeshesFromPoints::InAppendMeshPinLabel, EPCGDataType::DynamicMesh, false, false).SetRequiredPin();
+	}
+	
+	return Properties;
+}
+
+TArray<FPCGPinProperties> UPCGAppendMeshesFromPointsSettings::OutputPinProperties() const
+{ 
+	TArray<FPCGPinProperties> Properties;
+	Properties.Emplace(PCGPinConstants::DefaultOutputLabel, EPCGDataType::DynamicMesh, false, false);
+	return Properties;
+}
+
+bool FPCGAppendMeshesFromPointsElement::CanExecuteOnlyOnMainThread(FPCGContext* Context) const
+{
+	// Without context, we can't know, so force it in the main thread to be safe.
+	return !Context || Context->CurrentPhase == EPCGExecutionPhase::PrepareData;
+}
+
+FPCGContext* FPCGAppendMeshesFromPointsElement::CreateContext()
+{
+	return new FPCGAppendMeshesFromPointsContext();
+}
+
+bool FPCGAppendMeshesFromPointsElement::PrepareDataInternal(FPCGContext* InContext) const
+{
+	TRACE_CPUPROFILER_EVENT_SCOPE(FPCGAppendMeshesFromPointsElement::Execute);
+
+	FPCGAppendMeshesFromPointsContext* Context = static_cast<FPCGAppendMeshesFromPointsContext*>(InContext);
+	check(Context);
+	
+	const UPCGAppendMeshesFromPointsSettings* Settings = InContext->GetInputSettings<UPCGAppendMeshesFromPointsSettings>();
+	check(Settings);
+	
+	if (Context->WasLoadRequested())
+	{
+		return true;
+	}
+
+	if (Settings->Mode == EPCGAppendMeshesFromPointsMode::SingleStaticMesh)
+	{
+		if (Settings->StaticMesh.IsNull())
+		{
+			return true;
+		}
+		
+		Context->bPrepareDataSucceeded = true;
+		return Context->RequestResourceLoad(Context, {Settings->StaticMesh.ToSoftObjectPath()}, !Settings->bSynchronousLoad);
+	}
+
+	if (Settings->Mode == EPCGAppendMeshesFromPointsMode::StaticMeshFromAttribute)
+	{
+		const TArray<FPCGTaggedData> InputPoints = InContext->InputData.GetInputsByPin(PCGAppendMeshesFromPoints::InPointsPinLabel);
+		const UPCGPointData* InPointData = !InputPoints.IsEmpty() ? Cast<const UPCGPointData>(InputPoints[0].Data) : nullptr;
+
+		if (!InPointData)
+		{
+			return true;
+		}
+		
+		const FPCGAttributePropertyInputSelector Selector = Settings->MeshAttribute.CopyAndFixLast(InPointData);
+		const TUniquePtr<const IPCGAttributeAccessor> Accessor = PCGAttributeAccessorHelpers::CreateConstAccessor(InPointData, Selector);
+		const TUniquePtr<const IPCGAttributeAccessorKeys> Keys = PCGAttributeAccessorHelpers::CreateConstKeys(InPointData, Selector);
+
+		if (!Accessor || !Keys)
+		{
+			PCGLog::Metadata::LogFailToCreateAccessorError(Selector, Context);
+			return true;
+		}
+
+		if (!PCG::Private::IsBroadcastableOrConstructible(Accessor->GetUnderlyingType(), PCG::Private::MetadataTypes<FSoftObjectPath>::Id))
+		{
+			PCGLog::Metadata::LogFailToGetAttributeError<FSoftObjectPath>(Selector, Accessor.Get(), Context);
+			return true;
+		}
+
+		TArray<FSoftObjectPath> StaticMeshesToLoad;
+		PCGMetadataElementCommon::ApplyOnAccessor<FSoftObjectPath>(*Keys, *Accessor, [Context, &StaticMeshesToLoad](const FSoftObjectPath& Path, int32 Index)
+		{
+			if (!Path.IsNull())
+			{
+				StaticMeshesToLoad.AddUnique(Path);
+				Context->MeshToPointIndicesMapping.FindOrAdd(Path).Add(Index);
+			}
+		}, EPCGAttributeAccessorFlags::AllowBroadcastAndConstructible);
+
+		if (StaticMeshesToLoad.IsEmpty())
+		{
+			return true;
+		}
+
+		Context->bPrepareDataSucceeded = true;
+		return Context->RequestResourceLoad(Context, std::move(StaticMeshesToLoad), !Settings->bSynchronousLoad);
+	}
+	
+	Context->bPrepareDataSucceeded = true;
+	return true;
+}
+
+bool FPCGAppendMeshesFromPointsElement::ExecuteInternal(FPCGContext* InContext) const
+{
+	TRACE_CPUPROFILER_EVENT_SCOPE(FPCGAppendMeshesFromPointsElement::Execute);
+
+	FPCGAppendMeshesFromPointsContext* Context = static_cast<FPCGAppendMeshesFromPointsContext*>(InContext);
+	check(Context);
+
+	const UPCGAppendMeshesFromPointsSettings* Settings = InContext->GetInputSettings<UPCGAppendMeshesFromPointsSettings>();
+	check(Settings);
+
+	const TArray<FPCGTaggedData> InputPoints = InContext->InputData.GetInputsByPin(PCGAppendMeshesFromPoints::InPointsPinLabel);
+	const TArray<FPCGTaggedData> InputDynMesh = InContext->InputData.GetInputsByPin(PCGAppendMeshesFromPoints::InDynMeshPinLabel);
+	const TArray<FPCGTaggedData> InputAppendDynMesh = InContext->InputData.GetInputsByPin(PCGAppendMeshesFromPoints::InAppendMeshPinLabel);
+
+	const UPCGPointData* InPointData = !InputPoints.IsEmpty() ? Cast<const UPCGPointData>(InputPoints[0].Data) : nullptr;
+	const UPCGDynamicMeshData* InDynMeshData = !InputDynMesh.IsEmpty() ? Cast<const UPCGDynamicMeshData>(InputDynMesh[0].Data) : nullptr;
+	const UPCGDynamicMeshData* InAppendDynMeshData = !InputAppendDynMesh.IsEmpty() ? Cast<const UPCGDynamicMeshData>(InputAppendDynMesh[0].Data) : nullptr;
+
+	if (!InPointData || !InDynMeshData || !Context->bPrepareDataSucceeded)
+	{
+		InContext->OutputData.TaggedData = InputDynMesh;
+		return true;
+	}
+
+	FPCGTaggedData& OutputData = InContext->OutputData.TaggedData.Emplace_GetRef(InputDynMesh[0]);
+	if (InPointData->GetPoints().IsEmpty())
+	{
+		return true;
+	}
+
+	UPCGDynamicMeshData* OutDynMeshData = CopyOrSteal(InputDynMesh[0], InContext);
+	if (!OutDynMeshData)
+	{
+		return true;
+	}
+	
+	OutputData.Data = OutDynMeshData;
+
+	TMap<FSoftObjectPath, UE::Geometry::FDynamicMesh3> StaticMeshToDynMesh;
+
+	auto ConvertStaticMesh = [Settings, &StaticMeshToDynMesh, InContext](const FSoftObjectPath& InPath) -> bool
+	{
+		UE::Conversion::FStaticMeshConversionOptions ConversionOptions{};
+		FText ErrorMessage;
+		UE::Conversion::EMeshLODType LODType = static_cast<UE::Conversion::EMeshLODType>(Settings->RequestedLODType);
+		
+		UE::Geometry::FDynamicMesh3& NewMesh = StaticMeshToDynMesh.Add(InPath);
+	
+		if (!UE::Conversion::StaticMeshToDynamicMesh(Settings->StaticMesh.Get(), NewMesh, ErrorMessage, ConversionOptions, LODType, Settings->RequestedLODIndex))
+		{
+			PCGLog::LogErrorOnGraph(ErrorMessage, InContext);
+			return false;
+		}
+		else
+		{
+			return true;
+		}
+	};
+
+	switch (Settings->Mode)
+	{
+	case EPCGAppendMeshesFromPointsMode::SingleStaticMesh:
+	{
+		if (!ConvertStaticMesh(Settings->StaticMesh.ToSoftObjectPath()))
+		{
+			return true;
+		}
+
+		break;
+	}
+	case EPCGAppendMeshesFromPointsMode::StaticMeshFromAttribute:
+	{
+		StaticMeshToDynMesh.Reserve(Context->MeshToPointIndicesMapping.Num());
+	
+		for (TPair<FSoftObjectPath, TArray<int32>>& It : Context->MeshToPointIndicesMapping)
+		{
+			if (!ConvertStaticMesh(It.Key))
+			{
+				return true;
+			}
+		}
+
+		break;
+	}
+	case EPCGAppendMeshesFromPointsMode::DynamicMesh:
+	{
+		if (!InAppendDynMeshData)
+		{
+			PCGLog::InputOutput::LogTypedDataNotFoundWarning(EPCGDataType::DynamicMesh, PCGAppendMeshesFromPoints::InAppendMeshPinLabel, Context);
+			return true;
+		}
+
+		break;
+	}
+	default:
+		return true;
+	}
+	
+	UE::Geometry::FMeshIndexMappings MeshIndexMappings;
+	UE::Geometry::FDynamicMeshEditor Editor(OutDynMeshData->GetMutableDynamicMesh()->GetMeshPtr());
+
+	auto AppendMesh = [&Editor, &MeshIndexMappings](const UE::Geometry::FDynamicMesh3* MeshToAppend, const FPCGPoint& Point)
+	{
+		Editor.AppendMesh(MeshToAppend, MeshIndexMappings,
+		[&Point](int, const FVector& Position) { return Point.Transform.TransformPosition(Position); },
+		[&Point](int, const FVector& Normal)
+		{
+			const FVector& PointScale = Point.Transform.GetScale3D();
+			const double DetSign = FMathd::SignNonZero(PointScale.X * PointScale.Y * PointScale.Z); // we only need to multiply by the sign of the determinant, rather than divide by it, since we normalize later anyway
+			const FVector SafeInversePointScale(PointScale.Y * PointScale.Z * DetSign, PointScale.X * PointScale.Z * DetSign, PointScale.X * PointScale.Y * DetSign);
+			return Point.Transform.TransformVectorNoScale((SafeInversePointScale * Normal).GetSafeNormal());
+		});
+	};
+
+	const TArray<FPCGPoint>& InPoints = InPointData->GetPoints();
+
+	if (Settings->Mode == EPCGAppendMeshesFromPointsMode::StaticMeshFromAttribute)
+	{
+		for (TPair<FSoftObjectPath, TArray<int32>>& It : Context->MeshToPointIndicesMapping)
+		{
+			const UE::Geometry::FDynamicMesh3* MeshToAppend = &StaticMeshToDynMesh[It.Key];
+			check(MeshToAppend);
+
+			for (const int32 i : It.Value)
+			{
+				AppendMesh(MeshToAppend, InPoints[i]);
+			}
+		}
+	}
+	else
+	{
+		const UE::Geometry::FDynamicMesh3* MeshToAppend = nullptr;
+		if (Settings->Mode == EPCGAppendMeshesFromPointsMode::SingleStaticMesh)
+		{
+			MeshToAppend = &StaticMeshToDynMesh[Settings->StaticMesh.ToSoftObjectPath()];
+		}
+		else
+		{
+			MeshToAppend = InAppendDynMeshData->GetDynamicMesh() ? InAppendDynMeshData->GetDynamicMesh()->GetMeshPtr() : nullptr;
+		}
+
+		if (!MeshToAppend)
+		{
+			PCGLog::LogErrorOnGraph(LOCTEXT("InvalidDynMeshToAppend", "Invalid Dynamic mesh to append"), Context);
+			return true;
+		}
+
+		for (const FPCGPoint& Point : InPoints)
+		{
+			AppendMesh(MeshToAppend, Point);
+		}
+	}
+
+	return true;
+}
+
+#undef LOCTEXT_NAMESPACE
diff --git a/Engine/Plugins/PCGInterops/PCGGeometryScriptInterop/Source/PCGGeometryScriptInterop/Private/Elements/PCGBooleanOperation.cpp b/Engine/Plugins/PCGInterops/PCGGeometryScriptInterop/Source/PCGGeometryScriptInterop/Private/Elements/PCGBooleanOperation.cpp
new file mode 100644
--- /dev/null
+++ b/Engine/Plugins/PCGInterops/PCGGeometryScriptInterop/Source/PCGGeometryScriptInterop/Private/Elements/PCGBooleanOperation.cpp
@@ -0,0 +1,112 @@
+// Copyright Epic Games, Inc. All Rights Reserved.
+
+#include "Elements/PCGBooleanOperation.h"
+
+#include "PCGContext.h"
+#include "Data/PCGDynamicMeshData.h"
+#include "Utils/PCGLogErrors.h"
+
+#define LOCTEXT_NAMESPACE "PCGBooleanOperationElement"
+
+namespace PCGBooleanOperation
+{
+	static const FName InputAPinLabel = TEXT("InA");
+	static const FName InputBPinLabel = TEXT("InB");
+}
+
+#if WITH_EDITOR
+FName UPCGBooleanOperationSettings::GetDefaultNodeName() const
+{
+	return FName(TEXT("BooleanOperation"));
+}
+
+FText UPCGBooleanOperationSettings::GetDefaultNodeTitle() const
+{
+	return LOCTEXT("NodeTitle", "Boolean Operation");
+}
+
+FText UPCGBooleanOperationSettings::GetNodeTooltipText() const
+{
+	return LOCTEXT("NodeTooltip", "Do a boolean operation between 2 dynamic meshes.");
+}
+#endif // WITH_EDITOR
+
+FPCGElementPtr UPCGBooleanOperationSettings::CreateElement() const
+{
+	return MakeShared<FPCGBooleanOperationElement>();
+}
+
+TArray<FPCGPinProperties> UPCGBooleanOperationSettings::InputPinProperties() const
+{
+	TArray<FPCGPinProperties> Properties;
+	Properties.Emplace_GetRef(PCGBooleanOperation::InputAPinLabel, EPCGDataType::DynamicMesh).SetRequiredPin();
+	Properties.Emplace_GetRef(PCGBooleanOperation::InputBPinLabel, EPCGDataType::DynamicMesh).SetRequiredPin();
+	return Properties;
+}
+
+bool FPCGBooleanOperationElement::ExecuteInternal(FPCGContext* InContext) const
+{
+	TRACE_CPUPROFILER_EVENT_SCOPE(FPCGBooleanOperationElement::Execute);
+
+	check(InContext);
+
+	const UPCGBooleanOperationSettings* Settings = InContext->GetInputSettings<UPCGBooleanOperationSettings>();
+	check(Settings);
+
+	const TArray<FPCGTaggedData> InputsA = InContext->InputData.GetInputsByPin(PCGBooleanOperation::InputAPinLabel);
+	const TArray<FPCGTaggedData> InputsB = InContext->InputData.GetInputsByPin(PCGBooleanOperation::InputBPinLabel);
+
+	if (InputsA.IsEmpty() || InputsB.IsEmpty())
+	{
+		return true;
+	}
+
+	if (Settings->bBoolEachAWithEveryB && InputsA.Num() != 1 && InputsB.Num() != 1 && InputsA.Num() != InputsB.Num())
+	{
+		PCGLog::LogErrorOnGraph(LOCTEXT("MismatchNumInputs", "There is a mismatch between the number of inputs. If BoolEachAWithEveryB is false, we only support N:1, 1:N and N:N operations"), InContext);
+		return true;
+	}
+
+	// We can only steal the input if input A is not used multiple times.
+	const bool bCanStealInput = InputsB.Num() == 1 || (!Settings->bBoolEachAWithEveryB && InputsA.Num() == InputsB.Num());
+	const int32 NumIterations = Settings->bBoolEachAWithEveryB ? InputsA.Num() * InputsB.Num() : FMath::Max(InputsA.Num(), InputsB.Num());
+
+	for (int32 i = 0; i < NumIterations; ++i)
+	{
+		const FPCGTaggedData& InputA = InputsA[Settings->bBoolEachAWithEveryB ? i / InputsB.Num() : i % InputsA.Num()];
+		const FPCGTaggedData& InputB = InputsB[Settings->bBoolEachAWithEveryB ? i % InputsB.Num() : i % InputsB.Num()];
+
+		const UPCGDynamicMeshData* InputMeshA = Cast<const UPCGDynamicMeshData>(InputA.Data);
+		const UPCGDynamicMeshData* InputMeshB = Cast<const UPCGDynamicMeshData>(InputB.Data);
+
+		if (!InputMeshA || !InputMeshB)
+		{
+			PCGLog::InputOutput::LogInvalidInputDataError(InContext);
+			continue;
+		}
+
+		UPCGDynamicMeshData* OutputMeshData = bCanStealInput ? CopyOrSteal(InputA, InContext) : CastChecked<UPCGDynamicMeshData>(InputMeshA->DuplicateData(InContext));
+		check(OutputMeshData);
+
+		// Second mesh is required to be non const, but it won't be modified (Geometry Script API is not const friendly), hence the const_cast.
+		UGeometryScriptLibrary_MeshBooleanFunctions::ApplyMeshBoolean(OutputMeshData->GetMutableDynamicMesh(), FTransform::Identity,
+	const_cast<UDynamicMesh*>(InputMeshB->GetDynamicMesh()), FTransform::Identity,
+			Settings->BooleanOperation, Settings->BooleanOperationOptions);
+		
+		FPCGTaggedData& OutputData = InContext->OutputData.TaggedData.Emplace_GetRef(InputA);
+		if (Settings->TagInheritanceMode == EPCGBooleanOperationTagInheritanceMode::B)
+		{
+			OutputData.Tags = InputB.Tags;
+		}
+		else if (Settings->TagInheritanceMode == EPCGBooleanOperationTagInheritanceMode::Both)
+		{
+			OutputData.Tags.Append(InputB.Tags);
+		}
+
+		OutputData.Data = OutputMeshData;
+	}
+	
+	return true;
+}
+
+#undef LOCTEXT_NAMESPACE
diff --git a/Engine/Plugins/PCGInterops/PCGGeometryScriptInterop/Source/PCGGeometryScriptInterop/Private/Elements/PCGCreateEmptyDynamicMesh.cpp b/Engine/Plugins/PCGInterops/PCGGeometryScriptInterop/Source/PCGGeometryScriptInterop/Private/Elements/PCGCreateEmptyDynamicMesh.cpp
new file mode 100644
--- /dev/null
+++ b/Engine/Plugins/PCGInterops/PCGGeometryScriptInterop/Source/PCGGeometryScriptInterop/Private/Elements/PCGCreateEmptyDynamicMesh.cpp
@@ -0,0 +1,58 @@
+// Copyright Epic Games, Inc. All Rights Reserved.
+
+#include "Elements/PCGCreateEmptyDynamicMesh.h"
+
+#include "PCGContext.h"
+#include "Data/PCGDynamicMeshData.h"
+
+#define LOCTEXT_NAMESPACE "PCGCreateEmptyDynamicMeshElement"
+
+#if WITH_EDITOR
+FName UPCGCreateEmptyDynamicMeshSettings::GetDefaultNodeName() const
+{
+	return FName(TEXT("CreateEmptyDynamicMesh"));
+}
+
+FText UPCGCreateEmptyDynamicMeshSettings::GetDefaultNodeTitle() const
+{
+	return LOCTEXT("NodeTitle", "Create Empty Dynamic Mesh");
+}
+
+FText UPCGCreateEmptyDynamicMeshSettings::GetNodeTooltipText() const
+{
+	return LOCTEXT("NodeTooltip", "Create an empty dynamic mesh data.");
+}
+#endif // WITH_EDITOR
+
+FPCGElementPtr UPCGCreateEmptyDynamicMeshSettings::CreateElement() const
+{
+	return MakeShared<FPCGCreateEmptyDynamicMeshElement>();
+}
+
+TArray<FPCGPinProperties> UPCGCreateEmptyDynamicMeshSettings::InputPinProperties() const
+{
+	return {};
+}
+
+TArray<FPCGPinProperties> UPCGCreateEmptyDynamicMeshSettings::OutputPinProperties() const
+{ 
+	TArray<FPCGPinProperties> Properties;
+	Properties.Emplace(PCGPinConstants::DefaultOutputLabel, EPCGDataType::DynamicMesh, false, false);
+	return Properties;
+}
+
+bool FPCGCreateEmptyDynamicMeshElement::ExecuteInternal(FPCGContext* InContext) const
+{
+	TRACE_CPUPROFILER_EVENT_SCOPE(FPCGCreateEmptyDynamicMeshElement::Execute);
+
+	check(InContext);
+
+	const UPCGCreateEmptyDynamicMeshSettings* Settings = InContext->GetInputSettings<UPCGCreateEmptyDynamicMeshSettings>();
+	check(Settings);
+	
+	InContext->OutputData.TaggedData.Emplace_GetRef().Data = FPCGContext::NewObject_AnyThread<UPCGDynamicMeshData>(InContext);
+
+	return true;
+}
+
+#undef LOCTEXT_NAMESPACE
diff --git a/Engine/Plugins/PCGInterops/PCGGeometryScriptInterop/Source/PCGGeometryScriptInterop/Private/Elements/PCGDynamicMeshBaseElement.cpp b/Engine/Plugins/PCGInterops/PCGGeometryScriptInterop/Source/PCGGeometryScriptInterop/Private/Elements/PCGDynamicMeshBaseElement.cpp
new file mode 100644
--- /dev/null
+++ b/Engine/Plugins/PCGInterops/PCGGeometryScriptInterop/Source/PCGGeometryScriptInterop/Private/Elements/PCGDynamicMeshBaseElement.cpp
@@ -0,0 +1,78 @@
+// Copyright Epic Games, Inc. All Rights Reserved.
+
+#include "Elements/PCGDynamicMeshBaseElement.h"
+
+#include "PCGContext.h"
+#include "PCGData.h"
+#include "Data/PCGDynamicMeshData.h"
+
+#define LOCTEXT_NAMESPACE "PCGDynamicMeshBaseElement"
+
+namespace PCGDynamicMeshBaseElement
+{
+	static TAutoConsoleVariable<bool> CVarPCGDynamicMeshAllowSteal(
+		TEXT("pcg.DynamamicMesh.AllowDataSteal"),
+		true,
+		TEXT("Allows to steal dynamic meshes pcg data, avoiding a copy when possible."));
+
+	static TAutoConsoleVariable<int32> CVarPCGDynamicMeshStealVerbose(
+		TEXT("pcg.DynamamicMesh.DataStealVerbose"),
+		0,
+		TEXT("Verbosity to track steal. 0 = None, 1 = Log, 2 = Log + Graph message"));
+}
+
+TArray<FPCGPinProperties> UPCGDynamicMeshBaseSettings::InputPinProperties() const
+{
+	TArray<FPCGPinProperties> Properties;
+	Properties.Emplace_GetRef(PCGPinConstants::DefaultInputLabel, EPCGDataType::DynamicMesh).SetRequiredPin();
+	return Properties;
+}
+
+TArray<FPCGPinProperties> UPCGDynamicMeshBaseSettings::OutputPinProperties() const
+{
+	TArray<FPCGPinProperties> Properties;
+	Properties.Emplace(PCGPinConstants::DefaultOutputLabel, EPCGDataType::DynamicMesh);
+	return Properties;
+}
+
+UPCGDynamicMeshData* IPCGDynamicMeshBaseElement::CopyOrSteal(const FPCGTaggedData& InTaggedData, FPCGContext* InContext)
+{
+	check(InContext);
+	
+	const UPCGDynamicMeshData* InData = Cast<const UPCGDynamicMeshData>(InTaggedData.Data);
+	if (!InData)
+	{
+		return nullptr;
+	}
+	
+	if (InTaggedData.bIsUsedMultipleTimes || !PCGDynamicMeshBaseElement::CVarPCGDynamicMeshAllowSteal.GetValueOnAnyThread())
+	{
+		return CastChecked<UPCGDynamicMeshData>(InData->DuplicateData(InContext));
+	}
+
+	const int32 StealVerbose = PCGDynamicMeshBaseElement::CVarPCGDynamicMeshStealVerbose.GetValueOnAnyThread();
+	if (StealVerbose != 0)
+	{
+		
+#if !UE_BUILD_SHIPPING
+		const FName OriginatingNodeName = InTaggedData.OriginatingNode.IsValid() ? InTaggedData.OriginatingNode->GetFName() : TEXT("*UnknownSource*");
+#else
+		const FName OriginatingNodeName = TEXT("*UnknownSource*");
+#endif // !UE_BUILD_SHIPPING
+		
+		const FText Message = FText::Format(LOCTEXT("StealVerbose", "[STEAL DATA] Data, originating from node {0}, on pin {1} was stolen."),
+			FText::FromName(OriginatingNodeName), FText::FromName(InTaggedData.Pin));
+		if (StealVerbose == 1)
+		{
+			PCGE_LOG_C(Warning, LogOnly, InContext, Message);
+		}
+		else
+		{
+			PCGLog::LogWarningOnGraph(Message, InContext);
+		}
+	}
+
+	return const_cast<UPCGDynamicMeshData*>(InData);
+}
+
+#undef LOCTEXT_NAMESPACE
diff --git a/Engine/Plugins/PCGInterops/PCGGeometryScriptInterop/Source/PCGGeometryScriptInterop/Private/Elements/PCGDynamicMeshTransform.cpp b/Engine/Plugins/PCGInterops/PCGGeometryScriptInterop/Source/PCGGeometryScriptInterop/Private/Elements/PCGDynamicMeshTransform.cpp
new file mode 100644
--- /dev/null
+++ b/Engine/Plugins/PCGInterops/PCGGeometryScriptInterop/Source/PCGGeometryScriptInterop/Private/Elements/PCGDynamicMeshTransform.cpp
@@ -0,0 +1,59 @@
+// Copyright Epic Games, Inc. All Rights Reserved.
+
+#include "Elements/PCGDynamicMeshTransform.h"
+
+#include "PCGContext.h"
+#include "Data/PCGDynamicMeshData.h"
+
+#include "UDynamicMesh.h"
+#include "DynamicMesh/MeshTransforms.h"
+
+#define LOCTEXT_NAMESPACE "PCGDynamicMeshTransformElement"
+
+#if WITH_EDITOR
+FName UPCGDynamicMeshTransformSettings::GetDefaultNodeName() const
+{
+	return FName(TEXT("DynamicMeshTransform"));
+}
+
+FText UPCGDynamicMeshTransformSettings::GetDefaultNodeTitle() const
+{
+	return LOCTEXT("NodeTitle", "Dynamic Mesh Transform");
+}
+
+FText UPCGDynamicMeshTransformSettings::GetNodeTooltipText() const
+{
+	return LOCTEXT("NodeTooltip", "Apply a transform to all dynamic meshes.");
+}
+#endif // WITH_EDITOR
+
+FPCGElementPtr UPCGDynamicMeshTransformSettings::CreateElement() const
+{
+	return MakeShared<FPCGDynamicMeshTransformElement>();
+}
+
+bool FPCGDynamicMeshTransformElement::ExecuteInternal(FPCGContext* InContext) const
+{
+	TRACE_CPUPROFILER_EVENT_SCOPE(FPCGDynamicMeshTransformElement::Execute);
+
+	check(InContext);
+
+	const UPCGDynamicMeshTransformSettings* Settings = InContext->GetInputSettings<UPCGDynamicMeshTransformSettings>();
+	check(Settings);
+
+	for (const FPCGTaggedData& Input : InContext->InputData.GetInputsByPin(PCGPinConstants::DefaultInputLabel))
+	{
+		UPCGDynamicMeshData* OutputData = CopyOrSteal(Input, InContext);
+		if (!OutputData)
+		{
+			continue;
+		}
+		
+		MeshTransforms::ApplyTransform(OutputData->GetMutableDynamicMesh()->GetMeshRef(), Settings->Transform);
+		InContext->OutputData.TaggedData.Emplace_GetRef(Input).Data = OutputData;
+	}
+
+	return true;
+}
+
+#undef LOCTEXT_NAMESPACE
diff --git a/Engine/Plugins/PCGInterops/PCGGeometryScriptInterop/Source/PCGGeometryScriptInterop/Private/Elements/PCGMergeDynamicMeshes.cpp b/Engine/Plugins/PCGInterops/PCGGeometryScriptInterop/Source/PCGGeometryScriptInterop/Private/Elements/PCGMergeDynamicMeshes.cpp
new file mode 100644
--- /dev/null
+++ b/Engine/Plugins/PCGInterops/PCGGeometryScriptInterop/Source/PCGGeometryScriptInterop/Private/Elements/PCGMergeDynamicMeshes.cpp
@@ -0,0 +1,79 @@
+// Copyright Epic Games, Inc. All Rights Reserved.
+
+#include "Elements/PCGMergeDynamicMeshes.h"
+
+#include "PCGContext.h"
+#include "Data/PCGDynamicMeshData.h"
+
+#include "DynamicMeshEditor.h"
+#include "UDynamicMesh.h"
+#include "GeometryScript/MeshBasicEditFunctions.h"
+
+#define LOCTEXT_NAMESPACE "PCGMergeDynamicMeshesElement"
+
+#if WITH_EDITOR
+FName UPCGMergeDynamicMeshesSettings::GetDefaultNodeName() const
+{
+	return FName(TEXT("MergeDynamicMeshes"));
+}
+
+FText UPCGMergeDynamicMeshesSettings::GetDefaultNodeTitle() const
+{
+	return LOCTEXT("NodeTitle", "Merge Dynamic Meshes");
+}
+
+FText UPCGMergeDynamicMeshesSettings::GetNodeTooltipText() const
+{
+	return LOCTEXT("NodeTooltip", "Appends all incoming dynamic meshes to the first dynamic mesh in order.");
+}
+#endif // WITH_EDITOR
+
+FPCGElementPtr UPCGMergeDynamicMeshesSettings::CreateElement() const
+{
+	return MakeShared<FPCGMergeDynamicMeshesElement>();
+}
+
+TArray<FPCGPinProperties> UPCGMergeDynamicMeshesSettings::OutputPinProperties() const
+{ 
+	TArray<FPCGPinProperties> Properties;
+	Properties.Emplace(PCGPinConstants::DefaultOutputLabel, EPCGDataType::DynamicMesh, false, false);
+	return Properties;
+}
+
+bool FPCGMergeDynamicMeshesElement::ExecuteInternal(FPCGContext* InContext) const
+{
+	TRACE_CPUPROFILER_EVENT_SCOPE(FPCGMergeDynamicMeshesElement::Execute);
+
+	check(InContext);
+
+	const UPCGMergeDynamicMeshesSettings* Settings = InContext->GetInputSettings<UPCGMergeDynamicMeshesSettings>();
+	check(Settings);
+
+	UPCGDynamicMeshData* OutputData = nullptr;
+
+	for (const FPCGTaggedData& Input : InContext->InputData.GetInputsByPin(PCGPinConstants::DefaultInputLabel))
+	{
+		const UPCGDynamicMeshData* InputData = Cast<const UPCGDynamicMeshData>(Input.Data);
+		if (!InputData)
+		{
+			PCGLog::InputOutput::LogInvalidInputDataError(InContext);
+			continue;
+		}
+
+		if (!OutputData)
+		{
+			OutputData = CopyOrSteal(Input, InContext);
+			InContext->OutputData.TaggedData.Emplace_GetRef(Input).Data = OutputData;
+		}
+		else
+		{
+			UE::Geometry::FMeshIndexMappings MeshIndexMappings;
+			UE::Geometry::FDynamicMeshEditor Editor(OutputData->GetMutableDynamicMesh()->GetMeshPtr());
+			Editor.AppendMesh(InputData->GetDynamicMesh()->GetMeshPtr(), MeshIndexMappings);
+		}
+	}
+	
+	return true;
+}
+
+#undef LOCTEXT_NAMESPACE
diff --git a/Engine/Plugins/PCGInterops/PCGGeometryScriptInterop/Source/PCGGeometryScriptInterop/Private/Elements/PCGSaveDynamicMeshToAsset.cpp b/Engine/Plugins/PCGInterops/PCGGeometryScriptInterop/Source/PCGGeometryScriptInterop/Private/Elements/PCGSaveDynamicMeshToAsset.cpp
new file mode 100644
--- /dev/null
+++ b/Engine/Plugins/PCGInterops/PCGGeometryScriptInterop/Source/PCGGeometryScriptInterop/Private/Elements/PCGSaveDynamicMeshToAsset.cpp
@@ -0,0 +1,124 @@
+// Copyright Epic Games, Inc. All Rights Reserved.
+
+#include "Elements/PCGSaveDynamicMeshToAsset.h"
+
+#include "PCGAssetExporterUtils.h"
+#include "PCGContext.h"
+#include "PCGParamData.h"
+#include "Data/PCGDynamicMeshData.h"
+#include "Utils/PCGLogErrors.h"
+
+#include "UDynamicMesh.h"
+#include "Engine/StaticMesh.h"
+#include "GeometryScript/MeshAssetFunctions.h"
+
+#define LOCTEXT_NAMESPACE "PCGSaveDynamicMeshToAssetElement"
+
+#if WITH_EDITOR
+FName UPCGSaveDynamicMeshToAssetSettings::GetDefaultNodeName() const
+{
+	return FName(TEXT("SaveDynamicMeshToAsset"));
+}
+
+FText UPCGSaveDynamicMeshToAssetSettings::GetDefaultNodeTitle() const
+{
+	return LOCTEXT("NodeTitle", "Save Dynamic Mesh To Asset");
+}
+
+FText UPCGSaveDynamicMeshToAssetSettings::GetNodeTooltipText() const
+{
+	return LOCTEXT("NodeTooltip", "Saves dynamic mesh data into a static mesh asset.");
+}
+#endif // WITH_EDITOR
+
+FPCGElementPtr UPCGSaveDynamicMeshToAssetSettings::CreateElement() const
+{
+	return MakeShared<FPCGSaveDynamicMeshToAssetElement>();
+}
+
+TArray<FPCGPinProperties> UPCGSaveDynamicMeshToAssetSettings::InputPinProperties() const
+{
+	TArray<FPCGPinProperties> Properties;
+	Properties.Emplace_GetRef(PCGPinConstants::DefaultInputLabel, EPCGDataType::DynamicMesh, false, false).SetRequiredPin();
+	return Properties;
+}
+
+TArray<FPCGPinProperties> UPCGSaveDynamicMeshToAssetSettings::OutputPinProperties() const
+{ 
+	TArray<FPCGPinProperties> Properties;
+	Properties.Emplace(TEXT("OutAssetPath"), EPCGDataType::Param, false, false);
+	return Properties;
+}
+
+bool FPCGSaveDynamicMeshToAssetElement::ExecuteInternal(FPCGContext* InContext) const
+{
+	TRACE_CPUPROFILER_EVENT_SCOPE(FPCGSaveDynamicMeshToAssetElement::Execute);
+
+	check(InContext);
+
+	// For now we only support writing to an asset in editor build.
+#if WITH_EDITOR
+	const UPCGSaveDynamicMeshToAssetSettings* Settings = InContext->GetInputSettings<UPCGSaveDynamicMeshToAssetSettings>();
+	check(Settings);
+
+	const TArray<FPCGTaggedData> Inputs = InContext->InputData.GetInputsByPin(PCGPinConstants::DefaultInputLabel);
+	if (Inputs.IsEmpty())
+	{
+		return true;
+	}
+	
+	if (Inputs.Num() > 1)
+	{
+		PCGLog::InputOutput::LogFirstInputOnlyWarning(PCGPinConstants::DefaultInputLabel, InContext);	
+	}
+
+	const UPCGDynamicMeshData* DynamicMeshData = Cast<const UPCGDynamicMeshData>(Inputs[0].Data);
+	if (!DynamicMeshData || !DynamicMeshData->GetDynamicMesh())
+	{
+		PCGLog::InputOutput::LogTypedDataNotFoundWarning(EPCGDataType::DynamicMesh, PCGPinConstants::DefaultInputLabel, InContext);
+		return true;
+	}
+
+	UObject* OutObject = nullptr;
+	
+	UPackage* OutPackage = UPCGAssetExporterUtils::CreateAsset<UStaticMesh>(Settings->ExportParams, [&DynamicMeshData, InContext, &OutObject, Settings](const FString&, UObject* Asset)
+	{
+		// Export options, make sure we don't emit a transaction.
+		FGeometryScriptCopyMeshToAssetOptions AssetOptions = Settings->CopyMeshToAssetOptions;
+		AssetOptions.bEmitTransaction = false;
+		AssetOptions.bDeferMeshPostEditChange = false;
+
+		// Convert merged mesh to static mesh, const_cast because GeometryScript API is not nice with constness, but it will not be modified.
+		EGeometryScriptOutcomePins OutResult;
+		UGeometryScriptLibrary_StaticMeshFunctions::CopyMeshToStaticMesh(const_cast<UDynamicMesh*>(DynamicMeshData->GetDynamicMesh()), CastChecked<UStaticMesh>(Asset), AssetOptions, Settings->MeshWriteLOD, OutResult);
+
+		if (OutResult == EGeometryScriptOutcomePins::Failure)
+		{
+			PCGLog::LogErrorOnGraph(LOCTEXT("ErrorConversion", "Error while converting dynamic mesh to static mesh."), InContext);
+			return false;
+		}
+		else
+		{
+			OutObject = Asset;
+			return true;
+		}
+	}, InContext);
+
+	if (!OutPackage || !OutObject)
+	{
+		return true;
+	}
+
+	UPCGParamData* OutParamData = NewObject<UPCGParamData>();
+	OutParamData->Metadata->CreateAttribute<FSoftObjectPath>(TEXT("AssetPath"), FSoftObjectPath(OutObject), false, false);
+	OutParamData->Metadata->AddEntry();
+
+	InContext->OutputData.TaggedData.Emplace_GetRef(Inputs[0]).Data = OutParamData;
+#else
+	PCGLog::LogWarningOnGraph(LOCTEXT("CannotExportInNonEditor", "Can't save a dynamic mesh to asset in non editor build."), InContext);
+#endif // WITH_EDITOR
+	
+	return true;
+}
+
+#undef LOCTEXT_NAMESPACE
diff --git a/Engine/Plugins/PCGInterops/PCGGeometryScriptInterop/Source/PCGGeometryScriptInterop/Private/Elements/PCGSpawnDynamicMesh.cpp b/Engine/Plugins/PCGInterops/PCGGeometryScriptInterop/Source/PCGGeometryScriptInterop/Private/Elements/PCGSpawnDynamicMesh.cpp
new file mode 100644
--- /dev/null
+++ b/Engine/Plugins/PCGInterops/PCGGeometryScriptInterop/Source/PCGGeometryScriptInterop/Private/Elements/PCGSpawnDynamicMesh.cpp
@@ -0,0 +1,83 @@
+// Copyright Epic Games, Inc. All Rights Reserved.
+
+#include "Elements/PCGSpawnDynamicMesh.h"
+
+#include "PCGContext.h"
+#include "Components/DynamicMeshComponent.h"
+#include "Data/PCGDynamicMeshData.h"
+#include "Resources/PCGDynamicMeshManagedComponent.h"
+
+#include "UDynamicMesh.h"
+#include "GameFramework/Actor.h"
+
+#define LOCTEXT_NAMESPACE "PCGSpawnDynamicMeshElement"
+
+#if WITH_EDITOR
+FName UPCGSpawnDynamicMeshSettings::GetDefaultNodeName() const
+{
+	return FName(TEXT("SpawnDynamicMesh"));
+}
+
+FText UPCGSpawnDynamicMeshSettings::GetDefaultNodeTitle() const
+{
+	return LOCTEXT("NodeTitle", "Spawn Dynamic Mesh");
+}
+
+FText UPCGSpawnDynamicMeshSettings::GetNodeTooltipText() const
+{
+	return LOCTEXT("NodeTooltip", "Spawn a dynamic mesh component for each dynamic mesh data in input.");
+}
+#endif // WITH_EDITOR
+
+FPCGElementPtr UPCGSpawnDynamicMeshSettings::CreateElement() const
+{
+	return MakeShared<FPCGSpawnDynamicMeshElement>();
+}
+
+bool FPCGSpawnDynamicMeshElement::ExecuteInternal(FPCGContext* InContext) const
+{
+	TRACE_CPUPROFILER_EVENT_SCOPE(FPCGSpawnDynamicMeshElement::Execute);
+
+	check(InContext);
+
+	const UPCGSpawnDynamicMeshSettings* Settings = InContext->GetInputSettings<UPCGSpawnDynamicMeshSettings>();
+	check(Settings);
+
+	AActor* TargetActor = Settings->TargetActor.IsValid() ? Settings->TargetActor.Get() : InContext->GetTargetActor(nullptr);
+	if (!TargetActor)
+	{
+		PCGLog::LogErrorOnGraph(LOCTEXT("InvalidTargetActor", "Invalid target actor."), InContext);
+		return true;
+	}
+
+	for (const FPCGTaggedData& Input : InContext->InputData.GetInputsByPin(PCGPinConstants::DefaultInputLabel))
+	{
+		const UPCGDynamicMeshData* DynMeshData = Cast<UPCGDynamicMeshData>(Input.Data);
+		if (!DynMeshData)
+		{
+			PCGLog::InputOutput::LogInvalidInputDataError(InContext);
+			continue;
+		}
+
+		UPCGDynamicMeshManagedComponent* ManagedComponent = PCGDynamicMeshManagedComponent::GetOrCreateDynamicMeshManagedComponent(InContext, Settings, DynMeshData, TargetActor);
+		UDynamicMeshComponent* Component = ManagedComponent ? ManagedComponent->GetComponent() : nullptr;
+		if (!Component)
+		{
+			continue;
+		}
+
+		// Needs to be a copy.
+		Component->SetMesh(UE::Geometry::FDynamicMesh3(DynMeshData->GetDynamicMesh()->GetMeshRef()));
+
+		for (const FString& Tag : Input.Tags)
+		{
+			Component->ComponentTags.AddUnique(*Tag);
+		}
+
+		InContext->OutputData.TaggedData.Emplace(Input);
+	}
+
+	return true;
+}
+
+#undef LOCTEXT_NAMESPACE
diff --git a/Engine/Plugins/PCGInterops/PCGGeometryScriptInterop/Source/PCGGeometryScriptInterop/Private/Elements/PCGStaticMeshToDynamicMeshElement.cpp b/Engine/Plugins/PCGInterops/PCGGeometryScriptInterop/Source/PCGGeometryScriptInterop/Private/Elements/PCGStaticMeshToDynamicMeshElement.cpp
new file mode 100644
--- /dev/null
+++ b/Engine/Plugins/PCGInterops/PCGGeometryScriptInterop/Source/PCGGeometryScriptInterop/Private/Elements/PCGStaticMeshToDynamicMeshElement.cpp
@@ -0,0 +1,120 @@
+// Copyright Epic Games, Inc. All Rights Reserved.
+
+#include "Elements/PCGStaticMeshToDynamicMeshElement.h"
+
+#include "PCGContext.h"
+#include "PCGModule.h"
+#include "Data/PCGDynamicMeshData.h"
+
+#include "UDynamicMesh.h"
+#include "ConversionUtils/SceneComponentToDynamicMesh.h"
+#include "Engine/StaticMesh.h"
+
+#define LOCTEXT_NAMESPACE "PCGStaticMeshToDynamicMeshElementElement"
+
+#if WITH_EDITOR
+FName UPCGStaticMeshToDynamicMeshSettings::GetDefaultNodeName() const
+{
+	return FName(TEXT("StaticMeshToDynamicMeshElement"));
+}
+
+FText UPCGStaticMeshToDynamicMeshSettings::GetDefaultNodeTitle() const
+{
+	return LOCTEXT("NodeTitle", "Static Mesh To Dynamic Mesh Element");
+}
+
+FText UPCGStaticMeshToDynamicMeshSettings::GetNodeTooltipText() const
+{
+	return LOCTEXT("NodeTooltip", "Convert a static mesh into a dynamic mesh data.");
+}
+#endif // WITH_EDITOR
+
+FPCGElementPtr UPCGStaticMeshToDynamicMeshSettings::CreateElement() const
+{
+	return MakeShared<FPCGStaticMeshToDynamicMeshElement>();
+}
+
+TArray<FPCGPinProperties> UPCGStaticMeshToDynamicMeshSettings::InputPinProperties() const
+{
+	return {};
+}
+
+TArray<FPCGPinProperties> UPCGStaticMeshToDynamicMeshSettings::OutputPinProperties() const
+{
+	TArray<FPCGPinProperties> Properties;
+	Properties.Emplace(PCGPinConstants::DefaultOutputLabel, EPCGDataType::DynamicMesh, false, false);
+	return Properties;
+}
+
+bool FPCGStaticMeshToDynamicMeshElement::CanExecuteOnlyOnMainThread(FPCGContext* Context) const
+{
+	// Without context, we can't know, so force it in the main thread to be safe.
+	return !Context || Context->CurrentPhase == EPCGExecutionPhase::PrepareData;
+}
+
+FPCGContext* FPCGStaticMeshToDynamicMeshElement::CreateContext()
+{
+	return new FPCGStaticMeshToDynamicMeshContext();
+}
+
+bool FPCGStaticMeshToDynamicMeshElement::PrepareDataInternal(FPCGContext* InContext) const
+{
+	TRACE_CPUPROFILER_EVENT_SCOPE(FPCGStaticMeshToDynamicMeshElement::Execute);
+
+	FPCGStaticMeshToDynamicMeshContext* Context = static_cast<FPCGStaticMeshToDynamicMeshContext*>(InContext);
+	check(Context);
+
+	const UPCGStaticMeshToDynamicMeshSettings* Settings = InContext->GetInputSettings<UPCGStaticMeshToDynamicMeshSettings>();
+	check(Settings);
+	
+	if (Context->WasLoadRequested() || Settings->StaticMesh.IsNull())
+	{
+		return true;
+	}
+	
+	return Context->RequestResourceLoad(Context, {Settings->StaticMesh.ToSoftObjectPath()}, !Settings->bSynchronousLoad);
+}
+
+bool FPCGStaticMeshToDynamicMeshElement::ExecuteInternal(FPCGContext* InContext) const
+{
+	TRACE_CPUPROFILER_EVENT_SCOPE(FPCGStaticMeshToDynamicMeshElement::Execute);
+
+	check(InContext);
+
+	const UPCGStaticMeshToDynamicMeshSettings* Settings = InContext->GetInputSettings<UPCGStaticMeshToDynamicMeshSettings>();
+	check(Settings);
+
+	UStaticMesh* StaticMesh = Settings->StaticMesh.Get();
+	if (!StaticMesh)
+	{
+		if (!Settings->StaticMesh.IsNull())
+		{
+			PCGLog::LogErrorOnGraph(LOCTEXT("StaticMeshNull", "Static mesh failed to load."), InContext);
+		}
+		
+		return true;
+	}
+	
+	UE::Conversion::EMeshLODType LODType = static_cast<UE::Conversion::EMeshLODType>(Settings->RequestedLODType);
+	int32 LODIndex = Settings->RequestedLODIndex;
+
+	UE::Conversion::FStaticMeshConversionOptions ConversionOptions{};
+	FText ErrorMessage;
+	UE::Geometry::FDynamicMesh3 NewMesh;
+	const bool bSuccess = UE::Conversion::StaticMeshToDynamicMesh(StaticMesh, NewMesh, ErrorMessage, ConversionOptions, LODType, LODIndex);
+
+	if (bSuccess)
+	{
+		UPCGDynamicMeshData* DynMeshData = FPCGContext::NewObject_AnyThread<UPCGDynamicMeshData>(InContext);
+		DynMeshData->Initialize(std::move(NewMesh), InContext);
+		InContext->OutputData.TaggedData.Emplace_GetRef().Data = DynMeshData;
+	}
+	else
+	{
+		PCGLog::LogErrorOnGraph(ErrorMessage, InContext);
+	}
+
+	return true;
+}
+
+#undef LOCTEXT_NAMESPACE
diff --git a/Engine/Plugins/PCGInterops/PCGGeometryScriptInterop/Source/PCGGeometryScriptInterop/Public/Elements/PCGAppendMeshesFromPoints.h b/Engine/Plugins/PCGInterops/PCGGeometryScriptInterop/Source/PCGGeometryScriptInterop/Public/Elements/PCGAppendMeshesFromPoints.h
new file mode 100644
--- /dev/null
+++ b/Engine/Plugins/PCGInterops/PCGGeometryScriptInterop/Source/PCGGeometryScriptInterop/Public/Elements/PCGAppendMeshesFromPoints.h
@@ -0,0 +1,82 @@
+// Copyright Epic Games, Inc. All Rights Reserved.
+
+#pragma once
+
+#include "PCGContext.h"
+#include "Async/PCGAsyncLoadingContext.h"
+#include "Elements/PCGDynamicMeshBaseElement.h"
+
+#include "GeometryScript/GeometryScriptTypes.h"
+
+#include "PCGAppendMeshesFromPoints.generated.h"
+
+UENUM(Blueprintable)
+enum class EPCGAppendMeshesFromPointsMode : uint8
+{
+	SingleStaticMesh UMETA(ToolTip="Mesh taken from the node settings"),
+	StaticMeshFromAttribute UMETA(ToolTip="Mesh taken from attributes on the points"),
+	DynamicMesh UMETA(ToolTip="Mesh taken from another dynamic mesh")
+};
+
+/**
+* Append meshes at the points transforms. Mesh can be a single static mesh, multiple meshes coming from the points or another dynamic mesh.
+*/
+UCLASS(MinimalAPI, BlueprintType, ClassGroup = (Procedural))
+class UPCGAppendMeshesFromPointsSettings : public UPCGDynamicMeshBaseSettings
+{
+	GENERATED_BODY()
+
+public:
+	//~Begin UPCGSettings interface
+#if WITH_EDITOR
+	virtual FName GetDefaultNodeName() const override;
+	virtual FText GetDefaultNodeTitle() const override;
+	virtual FText GetNodeTooltipText() const override;
+#endif
+
+protected:
+	virtual TArray<FPCGPinProperties> InputPinProperties() const override;
+	virtual TArray<FPCGPinProperties> OutputPinProperties() const override;
+	virtual FPCGElementPtr CreateElement() const override;
+#if WITH_EDITOR
+	virtual EPCGChangeType GetChangeTypeForProperty(const FName& InPropertyName) const override;
+#endif // WITH_EDITOR
+	//~End UPCGSettings interface
+
+public:
+	UPROPERTY(BlueprintReadWrite, EditAnywhere, Category="Settings")
+	EPCGAppendMeshesFromPointsMode Mode;
+	
+	UPROPERTY(BlueprintReadWrite, EditAnywhere, Category="Settings", meta = (EditCondition = "Mode==EPCGAppendMeshesFromPointsMode::SingleStaticMesh", EditConditionHides, PCG_Overridable))
+	TSoftObjectPtr<UStaticMesh> StaticMesh;
+
+	UPROPERTY(BlueprintReadWrite, EditAnywhere, Category="Settings", meta = (EditCondition = "Mode==EPCGAppendMeshesFromPointsMode::StaticMeshFromAttribute", EditConditionHides, PCG_Overridable))
+	FPCGAttributePropertyInputSelector MeshAttribute;
+	
+	/** LOD type to use when creating DynamicMesh from specified StaticMesh. */
+	UPROPERTY(BlueprintReadWrite, EditAnywhere, Category = "Settings|LODSettings", meta = (EditCondition = "Mode!=EPCGAppendMeshesFromPointsMode::DynamicMesh", EditConditionHides, PCG_Overridable))
+	EGeometryScriptLODType RequestedLODType = EGeometryScriptLODType::RenderData;
+
+	UPROPERTY(BlueprintReadWrite, EditAnywhere, Category = "Settings|LODSettings", meta = (EditCondition = "Mode!=EPCGAppendMeshesFromPointsMode::DynamicMesh", EditConditionHides, PCG_Overridable))
+	int32 RequestedLODIndex = 0;
+	
+	UPROPERTY(BlueprintReadWrite, EditAnywhere, Category = "Settings|Debug", meta = (EditCondition = "Mode!=EPCGAppendMeshesFromPointsMode::DynamicMesh", EditConditionHides))
+	bool bSynchronousLoad = false;
+};
+
+struct FPCGAppendMeshesFromPointsContext : public FPCGContext, public IPCGAsyncLoadingContext
+{
+	bool bPrepareDataSucceeded = false;
+	TMap<FSoftObjectPath, TArray<int32>> MeshToPointIndicesMapping;
+};
+
+class FPCGAppendMeshesFromPointsElement : public IPCGDynamicMeshBaseElement
+{
+public:
+	virtual bool CanExecuteOnlyOnMainThread(FPCGContext* Context) const override;
+protected:
+	virtual FPCGContext* CreateContext() override;
+	virtual bool PrepareDataInternal(FPCGContext* InContext) const override;
+	virtual bool ExecuteInternal(FPCGContext* InContext) const override;
+};
+
diff --git a/Engine/Plugins/PCGInterops/PCGGeometryScriptInterop/Source/PCGGeometryScriptInterop/Public/Elements/PCGBooleanOperation.h b/Engine/Plugins/PCGInterops/PCGGeometryScriptInterop/Source/PCGGeometryScriptInterop/Public/Elements/PCGBooleanOperation.h
new file mode 100644
--- /dev/null
+++ b/Engine/Plugins/PCGInterops/PCGGeometryScriptInterop/Source/PCGGeometryScriptInterop/Public/Elements/PCGBooleanOperation.h
@@ -0,0 +1,60 @@
+// Copyright Epic Games, Inc. All Rights Reserved.
+
+#pragma once
+
+#include "Elements/PCGDynamicMeshBaseElement.h"
+
+#include "GeometryScript/MeshBooleanFunctions.h"
+
+#include "PCGBooleanOperation.generated.h"
+
+UENUM(Blueprintable)
+enum class EPCGBooleanOperationTagInheritanceMode : uint8
+{
+	Both,
+	A,
+	B,
+};
+
+/**
+* Do a boolean operation between 2 dynamic meshes.
+*/
+UCLASS(MinimalAPI, BlueprintType, ClassGroup = (Procedural))
+class UPCGBooleanOperationSettings : public UPCGDynamicMeshBaseSettings
+{
+	GENERATED_BODY()
+
+public:
+	//~Begin UPCGSettings interface
+#if WITH_EDITOR
+	virtual FName GetDefaultNodeName() const override;
+	virtual FText GetDefaultNodeTitle() const override;
+	virtual FText GetNodeTooltipText() const override;
+#endif
+
+protected:
+	virtual TArray<FPCGPinProperties> InputPinProperties() const override;
+	virtual FPCGElementPtr CreateElement() const override;
+	//~End UPCGSettings interface
+	
+public:
+	UPROPERTY(BlueprintReadWrite, EditAnywhere, Category = Settings, meta = (PCG_Overridable))
+	EGeometryScriptBooleanOperation BooleanOperation = EGeometryScriptBooleanOperation::Intersection;
+	
+	UPROPERTY(BlueprintReadWrite, EditAnywhere, Category = Settings, meta = (PCG_Overridable))
+	FGeometryScriptMeshBooleanOptions BooleanOperationOptions;
+
+	UPROPERTY(BlueprintReadWrite, EditAnywhere, Category = Settings, meta = (PCG_Overridable))
+	EPCGBooleanOperationTagInheritanceMode TagInheritanceMode;
+	
+	/** Each dynamic mesh in input A will be boolean'd with every dyn mesh in input B (cartesian product), producing N * M dyn meshes. Otherwise, will do a N:N (or N:1 or 1:N) operation, producing N dynamic meshes. */
+	UPROPERTY(BlueprintReadWrite, EditAnywhere, Category = Settings, meta = (PCG_Overridable, DisplayName = "Bool Each A With Every B"))
+	bool bBoolEachAWithEveryB = false;
+};
+
+class FPCGBooleanOperationElement : public IPCGDynamicMeshBaseElement
+{
+protected:
+	virtual bool ExecuteInternal(FPCGContext* InContext) const override;
+};
+
diff --git a/Engine/Plugins/PCGInterops/PCGGeometryScriptInterop/Source/PCGGeometryScriptInterop/Public/Elements/PCGCreateEmptyDynamicMesh.h b/Engine/Plugins/PCGInterops/PCGGeometryScriptInterop/Source/PCGGeometryScriptInterop/Public/Elements/PCGCreateEmptyDynamicMesh.h
new file mode 100644
--- /dev/null
+++ b/Engine/Plugins/PCGInterops/PCGGeometryScriptInterop/Source/PCGGeometryScriptInterop/Public/Elements/PCGCreateEmptyDynamicMesh.h
@@ -0,0 +1,37 @@
+// Copyright Epic Games, Inc. All Rights Reserved.
+
+#pragma once
+
+#include "Elements/PCGDynamicMeshBaseElement.h"
+
+#include "PCGCreateEmptyDynamicMesh.generated.h"
+
+/**
+* Create an empty dynamic mesh data.
+*/
+UCLASS(MinimalAPI, BlueprintType, ClassGroup = (Procedural))
+class UPCGCreateEmptyDynamicMeshSettings : public UPCGDynamicMeshBaseSettings
+{
+	GENERATED_BODY()
+
+public:
+	//~Begin UPCGSettings interface
+#if WITH_EDITOR
+	virtual FName GetDefaultNodeName() const override;
+	virtual FText GetDefaultNodeTitle() const override;
+	virtual FText GetNodeTooltipText() const override;
+#endif
+
+protected:
+	virtual TArray<FPCGPinProperties> InputPinProperties() const override;
+	virtual TArray<FPCGPinProperties> OutputPinProperties() const override;
+	virtual FPCGElementPtr CreateElement() const override;
+	//~End UPCGSettings interface
+};
+
+class FPCGCreateEmptyDynamicMeshElement : public IPCGDynamicMeshBaseElement
+{
+protected:
+	virtual bool ExecuteInternal(FPCGContext* InContext) const override;
+};
+
diff --git a/Engine/Plugins/PCGInterops/PCGGeometryScriptInterop/Source/PCGGeometryScriptInterop/Public/Elements/PCGDynamicMeshBaseElement.h b/Engine/Plugins/PCGInterops/PCGGeometryScriptInterop/Source/PCGGeometryScriptInterop/Public/Elements/PCGDynamicMeshBaseElement.h
new file mode 100644
--- /dev/null
+++ b/Engine/Plugins/PCGInterops/PCGGeometryScriptInterop/Source/PCGGeometryScriptInterop/Public/Elements/PCGDynamicMeshBaseElement.h
@@ -0,0 +1,44 @@
+// Copyright Epic Games, Inc. All Rights Reserved.
+
+#pragma once
+
+#include "PCGSettings.h"
+
+#include "PCGDynamicMeshBaseElement.generated.h"
+
+struct FPCGContext;
+class UPCGDynamicMeshData;
+
+/**
+* Base class for Dynamic Mesh nodes
+*/
+UCLASS(Abstract, MinimalAPI, ClassGroup = (Procedural))
+class UPCGDynamicMeshBaseSettings : public UPCGSettings
+{
+	GENERATED_BODY()
+
+public:
+	//~Begin UPCGSettings interface
+#if WITH_EDITOR
+	virtual EPCGSettingsType GetType() const override { return EPCGSettingsType::DynamicMesh; }
+#endif
+
+protected:
+	PCGGEOMETRYSCRIPTINTEROP_API virtual TArray<FPCGPinProperties> InputPinProperties() const override;
+	PCGGEOMETRYSCRIPTINTEROP_API virtual TArray<FPCGPinProperties> OutputPinProperties() const override;
+	//~End UPCGSettings interface
+};
+
+/**
+* Base class for Dynamic Mesh elements
+*/
+class IPCGDynamicMeshBaseElement : public IPCGElement
+{
+protected:
+	/** All dyn mesh elements are not cacheable to benefit from stealing data. */
+	virtual bool IsCacheable(const UPCGSettings*) const override { return false; }
+	virtual bool ShouldVerifyIfOutputsAreUsedMultipleTimes(const UPCGSettings*) const override { return true; }
+
+	PCGGEOMETRYSCRIPTINTEROP_API static UPCGDynamicMeshData* CopyOrSteal(const FPCGTaggedData& InTaggedData, FPCGContext* InContext);
+};
+
diff --git a/Engine/Plugins/PCGInterops/PCGGeometryScriptInterop/Source/PCGGeometryScriptInterop/Public/Elements/PCGDynamicMeshTransform.h b/Engine/Plugins/PCGInterops/PCGGeometryScriptInterop/Source/PCGGeometryScriptInterop/Public/Elements/PCGDynamicMeshTransform.h
new file mode 100644
--- /dev/null
+++ b/Engine/Plugins/PCGInterops/PCGGeometryScriptInterop/Source/PCGGeometryScriptInterop/Public/Elements/PCGDynamicMeshTransform.h
@@ -0,0 +1,39 @@
+// Copyright Epic Games, Inc. All Rights Reserved.
+
+#pragma once
+
+#include "Elements/PCGDynamicMeshBaseElement.h"
+
+#include "PCGDynamicMeshTransform.generated.h"
+
+/**
+* Apply a transform to all dynamic meshes.
+*/
+UCLASS(MinimalAPI, BlueprintType, ClassGroup = (Procedural))
+class UPCGDynamicMeshTransformSettings : public UPCGDynamicMeshBaseSettings
+{
+	GENERATED_BODY()
+
+public:
+	//~Begin UPCGSettings interface
+#if WITH_EDITOR
+	virtual FName GetDefaultNodeName() const override;
+	virtual FText GetDefaultNodeTitle() const override;
+	virtual FText GetNodeTooltipText() const override;
+#endif
+
+protected:
+	virtual FPCGElementPtr CreateElement() const override;
+	//~End UPCGSettings interface
+	
+public:
+	UPROPERTY(BlueprintReadWrite, EditAnywhere, Category="Settings", meta = (PCG_Overridable))
+	FTransform Transform;
+};
+
+class FPCGDynamicMeshTransformElement : public IPCGDynamicMeshBaseElement
+{
+protected:
+	virtual bool ExecuteInternal(FPCGContext* InContext) const override;
+};
+
diff --git a/Engine/Plugins/PCGInterops/PCGGeometryScriptInterop/Source/PCGGeometryScriptInterop/Public/Elements/PCGMergeDynamicMeshes.h b/Engine/Plugins/PCGInterops/PCGGeometryScriptInterop/Source/PCGGeometryScriptInterop/Public/Elements/PCGMergeDynamicMeshes.h
new file mode 100644
--- /dev/null
+++ b/Engine/Plugins/PCGInterops/PCGGeometryScriptInterop/Source/PCGGeometryScriptInterop/Public/Elements/PCGMergeDynamicMeshes.h
@@ -0,0 +1,36 @@
+// Copyright Epic Games, Inc. All Rights Reserved.
+
+#pragma once
+
+#include "Elements/PCGDynamicMeshBaseElement.h"
+
+#include "PCGMergeDynamicMeshes.generated.h"
+
+/**
+* Appends all incoming dynamic meshes to the first dynamic mesh in order.
+*/
+UCLASS(MinimalAPI, BlueprintType, ClassGroup = (Procedural))
+class UPCGMergeDynamicMeshesSettings : public UPCGDynamicMeshBaseSettings
+{
+	GENERATED_BODY()
+
+public:
+	//~Begin UPCGSettings interface
+#if WITH_EDITOR
+	virtual FName GetDefaultNodeName() const override;
+	virtual FText GetDefaultNodeTitle() const override;
+	virtual FText GetNodeTooltipText() const override;
+#endif
+
+protected:
+	virtual TArray<FPCGPinProperties> OutputPinProperties() const override;
+	virtual FPCGElementPtr CreateElement() const override;
+	//~End UPCGSettings interface
+};
+
+class FPCGMergeDynamicMeshesElement : public IPCGDynamicMeshBaseElement
+{
+protected:
+	virtual bool ExecuteInternal(FPCGContext* InContext) const override;
+};
+
diff --git a/Engine/Plugins/PCGInterops/PCGGeometryScriptInterop/Source/PCGGeometryScriptInterop/Public/Elements/PCGSaveDynamicMeshToAsset.h b/Engine/Plugins/PCGInterops/PCGGeometryScriptInterop/Source/PCGGeometryScriptInterop/Public/Elements/PCGSaveDynamicMeshToAsset.h
new file mode 100644
--- /dev/null
+++ b/Engine/Plugins/PCGInterops/PCGGeometryScriptInterop/Source/PCGGeometryScriptInterop/Public/Elements/PCGSaveDynamicMeshToAsset.h
@@ -0,0 +1,52 @@
+// Copyright Epic Games, Inc. All Rights Reserved.
+
+#pragma once
+
+#include "PCGAssetExporter.h"
+#include "Elements/PCGDynamicMeshBaseElement.h"
+
+#include "GeometryScript/MeshAssetFunctions.h"
+
+#include "PCGSaveDynamicMeshToAsset.generated.h"
+
+/**
+* Saves dynamic mesh data into a static mesh asset.
+*/
+UCLASS(MinimalAPI, BlueprintType, ClassGroup = (Procedural))
+class UPCGSaveDynamicMeshToAssetSettings : public UPCGDynamicMeshBaseSettings
+{
+	GENERATED_BODY()
+
+public:
+	//~Begin UPCGSettings interface
+#if WITH_EDITOR
+	virtual FName GetDefaultNodeName() const override;
+	virtual FText GetDefaultNodeTitle() const override;
+	virtual FText GetNodeTooltipText() const override;
+#endif
+
+protected:
+	virtual TArray<FPCGPinProperties> InputPinProperties() const override;
+	virtual TArray<FPCGPinProperties> OutputPinProperties() const override;
+	virtual FPCGElementPtr CreateElement() const override;
+	//~End UPCGSettings interface
+	
+public:
+	UPROPERTY(BlueprintReadWrite, EditAnywhere, Category = Settings, meta = (ShowOnlyInnerProperties, PCG_Overridable))
+	FPCGAssetExporterParameters ExportParams;
+
+	UPROPERTY(BlueprintReadWrite, EditAnywhere, Category = Settings, meta = (PCG_Overridable))
+	FGeometryScriptCopyMeshToAssetOptions CopyMeshToAssetOptions;
+	
+	UPROPERTY(BlueprintReadWrite, EditAnywhere, Category = Settings, meta = (PCG_Overridable))
+	FGeometryScriptMeshWriteLOD MeshWriteLOD;
+};
+
+class FPCGSaveDynamicMeshToAssetElement : public IPCGElement
+{
+public:
+	virtual bool CanExecuteOnlyOnMainThread(FPCGContext* Context) const override { return true; }
+protected:
+	virtual bool ExecuteInternal(FPCGContext* InContext) const override;
+};
+
diff --git a/Engine/Plugins/PCGInterops/PCGGeometryScriptInterop/Source/PCGGeometryScriptInterop/Public/Elements/PCGSpawnDynamicMesh.h b/Engine/Plugins/PCGInterops/PCGGeometryScriptInterop/Source/PCGGeometryScriptInterop/Public/Elements/PCGSpawnDynamicMesh.h
new file mode 100644
--- /dev/null
+++ b/Engine/Plugins/PCGInterops/PCGGeometryScriptInterop/Source/PCGGeometryScriptInterop/Public/Elements/PCGSpawnDynamicMesh.h
@@ -0,0 +1,47 @@
+// Copyright Epic Games, Inc. All Rights Reserved.
+
+#pragma once
+
+#include "Elements/PCGDynamicMeshBaseElement.h"
+
+#include "PCGSpawnDynamicMesh.generated.h"
+
+class UDynamicMeshComponent;
+/**
+* Spawn a dynamic mesh component for each dynamic mesh data in input.
+*/
+UCLASS(MinimalAPI, BlueprintType, ClassGroup = (Procedural))
+class UPCGSpawnDynamicMeshSettings : public UPCGDynamicMeshBaseSettings
+{
+	GENERATED_BODY()
+
+public:
+	//~Begin UPCGSettings interface
+#if WITH_EDITOR
+	virtual FName GetDefaultNodeName() const override;
+	virtual FText GetDefaultNodeTitle() const override;
+	virtual FText GetNodeTooltipText() const override;
+#endif
+
+protected:
+	virtual FPCGElementPtr CreateElement() const override;
+	//~End UPCGSettings interface
+
+public:
+	UPROPERTY(BlueprintReadWrite, Category = Settings, meta = (PCG_Overridable))
+	TSoftObjectPtr<AActor> TargetActor;
+
+	/** Specify a list of functions to be called on the target actor after instances are spawned. Functions need to be parameter-less and with "CallInEditor" flag enabled. */
+	UPROPERTY(BlueprintReadWrite, EditAnywhere, Category = Settings)
+	TArray<FName> PostProcessFunctionNames;
+};
+
+class FPCGSpawnDynamicMeshElement : public IPCGElement
+{
+public:
+	virtual bool CanExecuteOnlyOnMainThread(FPCGContext*) const override { return true; }
+	virtual bool IsCacheable(const UPCGSettings*) const override { return false; }
+protected:
+	virtual bool ExecuteInternal(FPCGContext* InContext) const override;
+};
+
diff --git a/Engine/Plugins/PCGInterops/PCGGeometryScriptInterop/Source/PCGGeometryScriptInterop/Public/Elements/PCGStaticMeshToDynamicMeshElement.h b/Engine/Plugins/PCGInterops/PCGGeometryScriptInterop/Source/PCGGeometryScriptInterop/Public/Elements/PCGStaticMeshToDynamicMeshElement.h
new file mode 100644
--- /dev/null
+++ b/Engine/Plugins/PCGInterops/PCGGeometryScriptInterop/Source/PCGGeometryScriptInterop/Public/Elements/PCGStaticMeshToDynamicMeshElement.h
@@ -0,0 +1,63 @@
+// Copyright Epic Games, Inc. All Rights Reserved.
+
+#pragma once
+
+#include "PCGContext.h"
+#include "Async/PCGAsyncLoadingContext.h"
+#include "Elements/PCGDynamicMeshBaseElement.h"
+
+#include "GeometryScript/GeometryScriptTypes.h"
+
+#include "PCGStaticMeshToDynamicMeshElement.generated.h"
+
+class UStaticMesh;
+
+/**
+* Convert a static mesh into a dynamic mesh data.
+*/
+UCLASS(MinimalAPI, BlueprintType, ClassGroup = (Procedural))
+class UPCGStaticMeshToDynamicMeshSettings : public UPCGDynamicMeshBaseSettings
+{
+	GENERATED_BODY()
+
+public:
+	//~Begin UPCGSettings interface
+#if WITH_EDITOR
+	virtual FName GetDefaultNodeName() const override;
+	virtual FText GetDefaultNodeTitle() const override;
+	virtual FText GetNodeTooltipText() const override;
+#endif
+
+protected:
+	virtual TArray<FPCGPinProperties> InputPinProperties() const override;
+	virtual TArray<FPCGPinProperties> OutputPinProperties() const override;
+	virtual FPCGElementPtr CreateElement() const override;
+	//~End UPCGSettings interface
+
+public:
+	UPROPERTY(BlueprintReadWrite, EditAnywhere, Category="Settings", meta = (PCG_Overridable))
+	TSoftObjectPtr<UStaticMesh> StaticMesh;
+
+	/** LOD type to use when creating DynamicMesh from specified StaticMesh. */
+	UPROPERTY(BlueprintReadWrite, EditAnywhere, Category = "Settings|LODSettings", meta = (EditCondition = "Mode!=EPCGAppendMeshesFromPointsMode::DynamicMesh", EditConditionHides, PCG_Overridable))
+	EGeometryScriptLODType RequestedLODType = EGeometryScriptLODType::RenderData;
+
+	UPROPERTY(BlueprintReadWrite, EditAnywhere, Category = "Settings|LODSettings", meta = (EditCondition = "Mode!=EPCGAppendMeshesFromPointsMode::DynamicMesh", EditConditionHides, PCG_Overridable))
+	int32 RequestedLODIndex = 0;
+	
+	UPROPERTY(BlueprintReadWrite, EditAnywhere, Category = "Settings|Debug")
+	bool bSynchronousLoad = false;
+};
+
+struct FPCGStaticMeshToDynamicMeshContext : public FPCGContext, public IPCGAsyncLoadingContext {};
+
+class FPCGStaticMeshToDynamicMeshElement : public IPCGDynamicMeshBaseElement
+{
+public:
+	virtual bool CanExecuteOnlyOnMainThread(FPCGContext* Context) const override;
+protected:
+	virtual FPCGContext* CreateContext() override;
+	virtual bool PrepareDataInternal(FPCGContext* Context) const override;
+	virtual bool ExecuteInternal(FPCGContext* InContext) const override;
+};
+
