[PCG] Fix static analysis

#rnx

diff --git a/Engine/Plugins/PCGInterops/PCGGeometryScriptInterop/Source/PCGGeometryScriptInterop/Private/Elements/PCGBooleanOperation.cpp b/Engine/Plugins/PCGInterops/PCGGeometryScriptInterop/Source/PCGGeometryScriptInterop/Private/Elements/PCGBooleanOperation.cpp
--- a/Engine/Plugins/PCGInterops/PCGGeometryScriptInterop/Source/PCGGeometryScriptInterop/Private/Elements/PCGBooleanOperation.cpp
+++ b/Engine/Plugins/PCGInterops/PCGGeometryScriptInterop/Source/PCGGeometryScriptInterop/Private/Elements/PCGBooleanOperation.cpp
@@ -61,7 +61,7 @@ bool FPCGBooleanOperationElement::ExecuteInternal(FPCGContext* InContext) const
 		return true;
 	}
 
-	if (Settings->bBoolEachAWithEveryB && InputsA.Num() != 1 && InputsB.Num() != 1 && InputsA.Num() != InputsB.Num())
+	if (!Settings->bBoolEachAWithEveryB && InputsA.Num() != 1 && InputsB.Num() != 1 && InputsA.Num() != InputsB.Num())
 	{
 		PCGLog::LogErrorOnGraph(LOCTEXT("MismatchNumInputs", "There is a mismatch between the number of inputs. If BoolEachAWithEveryB is false, we only support N:1, 1:N and N:N operations"), InContext);
 		return true;
@@ -73,8 +73,8 @@ bool FPCGBooleanOperationElement::ExecuteInternal(FPCGContext* InContext) const
 
 	for (int32 i = 0; i < NumIterations; ++i)
 	{
-		const FPCGTaggedData& InputA = InputsA[Settings->bBoolEachAWithEveryB ? i / InputsB.Num() : i % InputsA.Num()];
-		const FPCGTaggedData& InputB = InputsB[Settings->bBoolEachAWithEveryB ? i % InputsB.Num() : i % InputsB.Num()];
+		const FPCGTaggedData& InputA = InputsA[Settings->bBoolEachAWithEveryB ? (i / InputsB.Num()) : (i % InputsA.Num())];
+		const FPCGTaggedData& InputB = InputsB[i % InputsB.Num()];
 
 		const UPCGDynamicMeshData* InputMeshA = Cast<const UPCGDynamicMeshData>(InputA.Data);
 		const UPCGDynamicMeshData* InputMeshB = Cast<const UPCGDynamicMeshData>(InputB.Data);
