[PCG] GeometryScriptInterop: make sure PCG module is loaded

#rnx

diff --git a/Engine/Plugins/PCGInterops/PCGGeometryScriptInterop/Source/PCGGeometryScriptInterop/Private/PCGGeometryScriptInteropModule.cpp b/Engine/Plugins/PCGInterops/PCGGeometryScriptInterop/Source/PCGGeometryScriptInterop/Private/PCGGeometryScriptInteropModule.cpp
--- a/Engine/Plugins/PCGInterops/PCGGeometryScriptInterop/Source/PCGGeometryScriptInterop/Private/PCGGeometryScriptInteropModule.cpp
+++ b/Engine/Plugins/PCGInterops/PCGGeometryScriptInterop/Source/PCGGeometryScriptInterop/Private/PCGGeometryScriptInteropModule.cpp
@@ -21,6 +21,8 @@ public:
 	
 	virtual void StartupModule() override
 	{
+		FModuleManager::Get().LoadModuleChecked(TEXT("PCG"));
+
 #if WITH_EDITOR
 		FPCGDataVisualizationRegistry& DataVisRegistry = FPCGModule::GetMutablePCGDataVisualizationRegistry();
 		DataVisRegistry.RegisterPCGDataVisualization(UPCGDynamicMeshData::StaticClass(), MakeUnique<const FPCGDynamicMeshDataVisualization>());
