[PCG] Fix crash when inspecting an empty dynamic mesh data.

* Trivial fix, if there is no points, create an empty data, not a null data.

#rnx

diff --git a/Engine/Plugins/PCGInterops/PCGGeometryScriptInterop/Source/PCGGeometryScriptInterop/Private/Data/PCGDynamicMeshData.cpp b/Engine/Plugins/PCGInterops/PCGGeometryScriptInterop/Source/PCGGeometryScriptInterop/Private/Data/PCGDynamicMeshData.cpp
--- a/Engine/Plugins/PCGInterops/PCGGeometryScriptInterop/Source/PCGGeometryScriptInterop/Private/Data/PCGDynamicMeshData.cpp
+++ b/Engine/Plugins/PCGInterops/PCGGeometryScriptInterop/Source/PCGGeometryScriptInterop/Private/Data/PCGDynamicMeshData.cpp
@@ -160,11 +160,6 @@ const UPCGPointData* UPCGDynamicMeshData::ToPointData(FPCGContext* Context, cons
 	UGeometryScriptLibrary_MeshQueryFunctions::GetAllVertexPositions(DynamicMesh, VerticesList, /*bSkipGaps=*/false, /*bHasVertexIDGaps=*/Dummy);
 
 	TArray<FVector>& Vertices = *VerticesList.List.Get();
-
-	if (Vertices.IsEmpty())
-	{
-		return nullptr;
-	}
 	
 	UPCGPointData* Data = FPCGContext::NewObject_AnyThread<UPCGPointData>(Context);
 	TArray<FPCGPoint>& Points = Data->GetMutablePoints();
