[PCG] Clarify usage of Boolean operation for PCG Dynamic Mesh

* From 2 different modes to 3 different modes
* Added tooltips to make it clearer.

#rnx

diff --git a/Engine/Plugins/PCGInterops/PCGGeometryScriptInterop/Source/PCGGeometryScriptInterop/Private/Elements/PCGBooleanOperation.cpp b/Engine/Plugins/PCGInterops/PCGGeometryScriptInterop/Source/PCGGeometryScriptInterop/Private/Elements/PCGBooleanOperation.cpp
--- a/Engine/Plugins/PCGInterops/PCGGeometryScriptInterop/Source/PCGGeometryScriptInterop/Private/Elements/PCGBooleanOperation.cpp
+++ b/Engine/Plugins/PCGInterops/PCGGeometryScriptInterop/Source/PCGGeometryScriptInterop/Private/Elements/PCGBooleanOperation.cpp
@@ -27,7 +27,7 @@ FText UPCGBooleanOperationSettings::GetDefaultNodeTitle() const
 
 FText UPCGBooleanOperationSettings::GetNodeTooltipText() const
 {
-	return LOCTEXT("NodeTooltip", "Do a boolean operation between 2 dynamic meshes.");
+	return LOCTEXT("NodeTooltip", "Boolean operation between dynamic meshes.");
 }
 #endif // WITH_EDITOR
 
@@ -61,20 +61,27 @@ bool FPCGBooleanOperationElement::ExecuteInternal(FPCGContext* InContext) const
 		return true;
 	}
 
-	if (!Settings->bBoolEachAWithEveryB && InputsA.Num() != 1 && InputsB.Num() != 1 && InputsA.Num() != InputsB.Num())
+	if (Settings->Mode == EPCGBooleanOperationMode::EachAWithEachB && InputsA.Num() != 1 && InputsB.Num() != 1 && InputsA.Num() != InputsB.Num())
 	{
 		PCGLog::LogErrorOnGraph(LOCTEXT("MismatchNumInputs", "There is a mismatch between the number of inputs. If BoolEachAWithEveryB is false, we only support N:1, 1:N and N:N operations"), InContext);
 		return true;
 	}
-
+	
 	// We can only steal the input if input A is not used multiple times.
-	const bool bCanStealInput = InputsB.Num() == 1 || (!Settings->bBoolEachAWithEveryB && InputsA.Num() == InputsB.Num());
-	const int32 NumIterations = Settings->bBoolEachAWithEveryB ? InputsA.Num() * InputsB.Num() : FMath::Max(InputsA.Num(), InputsB.Num());
+	const bool bCanStealInput = InputsB.Num() == 1 || (Settings->Mode == EPCGBooleanOperationMode::EachAWithEachBSequentially)
+		|| (Settings->Mode == EPCGBooleanOperationMode::EachAWithEachB && InputsA.Num() == InputsB.Num());
+	
+	const int32 NumIterations = Settings->Mode == EPCGBooleanOperationMode::EachAWithEachB ? FMath::Max(InputsA.Num(), InputsB.Num()) : InputsA.Num() * InputsB.Num();
+	UPCGDynamicMeshData* CurrentOutputMeshData = nullptr;
+	FPCGTaggedData* CurrentTaggedOutputData = nullptr;
 
 	for (int32 i = 0; i < NumIterations; ++i)
 	{
-		const FPCGTaggedData& InputA = InputsA[Settings->bBoolEachAWithEveryB ? (i / InputsB.Num()) : (i % InputsA.Num())];
-		const FPCGTaggedData& InputB = InputsB[i % InputsB.Num()];
+		const int32 InputAIndex = Settings->Mode == EPCGBooleanOperationMode::EachAWithEachB ? (i % InputsA.Num()) : (i / InputsB.Num());
+		const int32 InputBIndex = i % InputsB.Num();
+		
+		const FPCGTaggedData& InputA = InputsA[InputAIndex];
+		const FPCGTaggedData& InputB = InputsB[InputBIndex];
 
 		const UPCGDynamicMeshData* InputMeshA = Cast<const UPCGDynamicMeshData>(InputA.Data);
 		const UPCGDynamicMeshData* InputMeshB = Cast<const UPCGDynamicMeshData>(InputB.Data);
@@ -85,25 +92,32 @@ bool FPCGBooleanOperationElement::ExecuteInternal(FPCGContext* InContext) const
 			continue;
 		}
 
-		UPCGDynamicMeshData* OutputMeshData = bCanStealInput ? CopyOrSteal(InputA, InContext) : CastChecked<UPCGDynamicMeshData>(InputMeshA->DuplicateData(InContext));
-		check(OutputMeshData);
+		// At every loop, we update our current data form the input.
+		// In the EachAWithEachBSequentially case, only do it at the beginning of a new cycle (when we start at the first B again).
+		if (Settings->Mode != EPCGBooleanOperationMode::EachAWithEachBSequentially || InputBIndex == 0)
+		{
+			CurrentOutputMeshData = bCanStealInput ? CopyOrSteal(InputA, InContext) : CastChecked<UPCGDynamicMeshData>(InputMeshA->DuplicateData(InContext));
+			CurrentTaggedOutputData = &InContext->OutputData.TaggedData.Emplace_GetRef(InputA);
+			CurrentTaggedOutputData->Data = CurrentOutputMeshData;
+		}
+
+		check(CurrentOutputMeshData);
+		check(CurrentTaggedOutputData);
 
 		// Second mesh is required to be non const, but it won't be modified (Geometry Script API is not const friendly), hence the const_cast.
-		UGeometryScriptLibrary_MeshBooleanFunctions::ApplyMeshBoolean(OutputMeshData->GetMutableDynamicMesh(), FTransform::Identity,
+		UGeometryScriptLibrary_MeshBooleanFunctions::ApplyMeshBoolean(CurrentOutputMeshData->GetMutableDynamicMesh(), FTransform::Identity,
 	const_cast<UDynamicMesh*>(InputMeshB->GetDynamicMesh()), FTransform::Identity,
 			Settings->BooleanOperation, Settings->BooleanOperationOptions);
+
 		
-		FPCGTaggedData& OutputData = InContext->OutputData.TaggedData.Emplace_GetRef(InputA);
 		if (Settings->TagInheritanceMode == EPCGBooleanOperationTagInheritanceMode::B)
 		{
-			OutputData.Tags = InputB.Tags;
+			CurrentTaggedOutputData->Tags = InputB.Tags;
 		}
 		else if (Settings->TagInheritanceMode == EPCGBooleanOperationTagInheritanceMode::Both)
 		{
-			OutputData.Tags.Append(InputB.Tags);
+			CurrentTaggedOutputData->Tags.Append(InputB.Tags);
 		}
-
-		OutputData.Data = OutputMeshData;
 	}
 	
 	return true;
diff --git a/Engine/Plugins/PCGInterops/PCGGeometryScriptInterop/Source/PCGGeometryScriptInterop/Public/Elements/PCGBooleanOperation.h b/Engine/Plugins/PCGInterops/PCGGeometryScriptInterop/Source/PCGGeometryScriptInterop/Public/Elements/PCGBooleanOperation.h
--- a/Engine/Plugins/PCGInterops/PCGGeometryScriptInterop/Source/PCGGeometryScriptInterop/Public/Elements/PCGBooleanOperation.h
+++ b/Engine/Plugins/PCGInterops/PCGGeometryScriptInterop/Source/PCGGeometryScriptInterop/Public/Elements/PCGBooleanOperation.h
@@ -16,8 +16,16 @@ enum class EPCGBooleanOperationTagInheritanceMode : uint8
 	B,
 };
 
+UENUM(Blueprintable)
+enum class EPCGBooleanOperationMode : uint8
+{
+	EachAWithEachB UMETA(DisplayName="Each A With Each B", Tooltip="Each input in A is boolean'd with its associated input in B (A1 with B1, A2 with B2, etc...). Produces N outputs."),
+	EachAWithEachBSequentially UMETA(DisplayName="Each A With Each B Sequentially", Tooltip="Each input in A is boolean'd with every input in B sequentially. (A1 with B1 then with B2, A2 with B1 then B2, etc...). Produces N outputs."),
+	EachAWithEveryB UMETA(DisplayName="Each A With Every B", Tooltip="Each input in A is boolean'd with input in B individually (Cartesian product: A1 with B1, A1 with B2, A2 with B1, A2 with B2, etc...). Produces N * M outputs.")
+};
+
 /**
-* Do a boolean operation between 2 dynamic meshes.
+* Do a boolean operation between dynamic meshes.
 */
 UCLASS(MinimalAPI, BlueprintType, ClassGroup = (Procedural))
 class UPCGBooleanOperationSettings : public UPCGDynamicMeshBaseSettings
@@ -47,9 +55,8 @@ public:
 	UPROPERTY(BlueprintReadWrite, EditAnywhere, Category = Settings, meta = (PCG_Overridable))
 	EPCGBooleanOperationTagInheritanceMode TagInheritanceMode;
 	
-	/** Each dynamic mesh in input A will be boolean'd with every dyn mesh in input B (cartesian product), producing N * M dyn meshes. Otherwise, will do a N:N (or N:1 or 1:N) operation, producing N dynamic meshes. */
-	UPROPERTY(BlueprintReadWrite, EditAnywhere, Category = Settings, meta = (PCG_Overridable, DisplayName = "Bool Each A With Every B"))
-	bool bBoolEachAWithEveryB = false;
+	UPROPERTY(BlueprintReadWrite, EditAnywhere, Category = Settings, meta = (PCG_Overridable))
+	EPCGBooleanOperationMode Mode = EPCGBooleanOperationMode::EachAWithEachB;
 };
 
 class FPCGBooleanOperationElement : public IPCGDynamicMeshBaseElement
