[PCG] Support for materials in Dynamic Meshes Data

#rnx

diff --git a/Engine/Plugins/PCGInterops/PCGGeometryScriptInterop/Source/PCGGeometryScriptInterop/Private/Data/PCGDynamicMeshData.cpp b/Engine/Plugins/PCGInterops/PCGGeometryScriptInterop/Source/PCGGeometryScriptInterop/Private/Data/PCGDynamicMeshData.cpp
--- a/Engine/Plugins/PCGInterops/PCGGeometryScriptInterop/Source/PCGGeometryScriptInterop/Private/Data/PCGDynamicMeshData.cpp
+++ b/Engine/Plugins/PCGInterops/PCGGeometryScriptInterop/Source/PCGGeometryScriptInterop/Private/Data/PCGDynamicMeshData.cpp
@@ -8,6 +8,7 @@
 #include "Elements/PCGVolumeSampler.h"
 
 #include "UDynamicMesh.h"
+#include "Components/DynamicMeshComponent.h"
 #include "GeometryScript/MeshQueryFunctions.h"
 #include "Helpers/PCGHelpers.h"
 #include "Serialization/ArchiveCrc32.h"
@@ -18,27 +19,29 @@ UPCGDynamicMeshData::UPCGDynamicMeshData(const FObjectInitializer& ObjectInitial
 	DynamicMesh = ObjectInitializer.CreateDefaultSubobject<UDynamicMesh>(this, TEXT("DynamicMesh"));
 }
 
-void UPCGDynamicMeshData::Initialize(UDynamicMesh* InMesh, FPCGContext* Context, bool bCanTakeOwnership)
+void UPCGDynamicMeshData::Initialize(UDynamicMesh* InMesh, bool bCanTakeOwnership, const TArray<UMaterialInterface*>& InMaterials)
 {
 	check(DynamicMesh);
 	if (bCanTakeOwnership)
 	{
-		Initialize(std::move(InMesh->GetMeshRef()));
+		Initialize(std::move(InMesh->GetMeshRef()), InMaterials);
 	}
 	else
 	{
 		// Make a copy
-		Initialize(UE::Geometry::FDynamicMesh3(InMesh->GetMeshRef()));
+		Initialize(UE::Geometry::FDynamicMesh3(InMesh->GetMeshRef()), InMaterials);
 	}
 }
 
-void UPCGDynamicMeshData::Initialize(UE::Geometry::FDynamicMesh3&& InMesh, FPCGContext* Context)
+void UPCGDynamicMeshData::Initialize(UE::Geometry::FDynamicMesh3&& InMesh, const TArray<UMaterialInterface*>& InMaterials)
 {
 	check(DynamicMesh);
 	DynamicMesh->SetMesh(std::move(InMesh));
 	ResetBounds();
 	
 	bDynamicMeshOctreeIsDirty = true;
+	
+	SetMaterials(InMaterials);
 }
 
 void UPCGDynamicMeshData::ResetBounds() const
@@ -181,11 +184,29 @@ const UPCGPointData* UPCGDynamicMeshData::ToPointData(FPCGContext* Context, cons
 	return Data;
 }
 
+void UPCGDynamicMeshData::SetMaterials(const TArray<UMaterialInterface*>& InMaterials)
+{
+	Materials.Reset(InMaterials.Num());
+	Materials.Append(InMaterials);
+}
+
+void UPCGDynamicMeshData::InitializeDynamicMeshComponentFromData(UDynamicMeshComponent* InComponent) const
+{
+	check(InComponent);
+
+	InComponent->SetMesh(UE::Geometry::FDynamicMesh3(DynamicMesh->GetMeshRef()));
+	for (int32 i = 0; i < Materials.Num(); ++i)
+	{
+		InComponent->SetMaterial(i, Materials[i]);
+	}
+}
+
 UPCGSpatialData* UPCGDynamicMeshData::CopyInternal(FPCGContext* Context) const
 {
 	UPCGDynamicMeshData* NewDynamicMeshData = FPCGContext::NewObject_AnyThread<UPCGDynamicMeshData>(Context);
 	
-	NewDynamicMeshData->Initialize(DynamicMesh, Context, /*bCanTakeOwnership=*/false);
+	NewDynamicMeshData->Initialize(DynamicMesh, /*bCanTakeOwnership=*/false);
+	NewDynamicMeshData->Materials = Materials;
 
 	return NewDynamicMeshData;
 }
diff --git a/Engine/Plugins/PCGInterops/PCGGeometryScriptInterop/Source/PCGGeometryScriptInterop/Private/Editor/PCGDynamicMeshDataVisualization.cpp b/Engine/Plugins/PCGInterops/PCGGeometryScriptInterop/Source/PCGGeometryScriptInterop/Private/Editor/PCGDynamicMeshDataVisualization.cpp
--- a/Engine/Plugins/PCGInterops/PCGGeometryScriptInterop/Source/PCGGeometryScriptInterop/Private/Editor/PCGDynamicMeshDataVisualization.cpp
+++ b/Engine/Plugins/PCGInterops/PCGGeometryScriptInterop/Source/PCGGeometryScriptInterop/Private/Editor/PCGDynamicMeshDataVisualization.cpp
@@ -23,11 +23,11 @@ void FPCGDynamicMeshDataVisualization::ExecuteDebugDisplay(FPCGContext* Context,
 
 	// We force debug resources to be transient.
 	UPCGDynamicMeshManagedComponent* ManagedComponent = PCGDynamicMeshManagedComponent::GetOrCreateDynamicMeshManagedComponent(Context, SettingsInterface, DynMeshData, TargetActor, EPCGEditorDirtyMode::Preview);
+	UDynamicMeshComponent* Component = ManagedComponent ? ManagedComponent->GetComponent() : nullptr;
 	
-	if (ManagedComponent && ManagedComponent->GetComponent())
+	if (Component)
 	{
-		// Needs to be a copy.
-		ManagedComponent->GetComponent()->SetMesh(UE::Geometry::FDynamicMesh3(DynMeshData->GetDynamicMesh()->GetMeshRef()));
+		DynMeshData->InitializeDynamicMeshComponentFromData(Component);
 	}
 }
 
diff --git a/Engine/Plugins/PCGInterops/PCGGeometryScriptInterop/Source/PCGGeometryScriptInterop/Private/Elements/PCGAppendMeshesFromPoints.cpp b/Engine/Plugins/PCGInterops/PCGGeometryScriptInterop/Source/PCGGeometryScriptInterop/Private/Elements/PCGAppendMeshesFromPoints.cpp
--- a/Engine/Plugins/PCGInterops/PCGGeometryScriptInterop/Source/PCGGeometryScriptInterop/Private/Elements/PCGAppendMeshesFromPoints.cpp
+++ b/Engine/Plugins/PCGInterops/PCGGeometryScriptInterop/Source/PCGGeometryScriptInterop/Private/Elements/PCGAppendMeshesFromPoints.cpp
@@ -6,6 +6,7 @@
 #include "Data/PCGDynamicMeshData.h"
 #include "Data/PCGPointData.h"
 #include "Elements/Metadata/PCGMetadataElementCommon.h"
+#include "Helpers/PCGGeometryHelpers.h"
 #include "Metadata/Accessors/PCGAttributeAccessorHelpers.h"
 
 #include "DynamicMeshEditor.h"
@@ -13,6 +14,7 @@
 #include "ConversionUtils/SceneComponentToDynamicMesh.h"
 #include "DynamicMesh/MeshIndexMappings.h"
 #include "Engine/StaticMesh.h"
+#include "GeometryScript/MeshAssetFunctions.h"
 
 #define LOCTEXT_NAMESPACE "PCGAppendMeshesFromPointsElement"
 
@@ -205,23 +207,32 @@ bool FPCGAppendMeshesFromPointsElement::ExecuteInternal(FPCGContext* InContext)
 
 	TMap<FSoftObjectPath, UE::Geometry::FDynamicMesh3> StaticMeshToDynMesh;
 
-	auto ConvertStaticMesh = [Settings, &StaticMeshToDynMesh, InContext](const FSoftObjectPath& InPath) -> bool
+	auto ConvertStaticMesh = [Settings, &StaticMeshToDynMesh, InContext, &OutDynMeshData](const FSoftObjectPath& InPath) -> bool
 	{
 		UE::Conversion::FStaticMeshConversionOptions ConversionOptions{};
 		FText ErrorMessage;
 		UE::Conversion::EMeshLODType LODType = static_cast<UE::Conversion::EMeshLODType>(Settings->RequestedLODType);
 		
 		UE::Geometry::FDynamicMesh3& NewMesh = StaticMeshToDynMesh.Add(InPath);
+		UStaticMesh* StaticMesh = Settings->StaticMesh.Get();
 	
-		if (!UE::Conversion::StaticMeshToDynamicMesh(Settings->StaticMesh.Get(), NewMesh, ErrorMessage, ConversionOptions, LODType, Settings->RequestedLODIndex))
+		if (!UE::Conversion::StaticMeshToDynamicMesh(StaticMesh, NewMesh, ErrorMessage, ConversionOptions, LODType, Settings->RequestedLODIndex))
 		{
 			PCGLog::LogErrorOnGraph(ErrorMessage, InContext);
 			return false;
 		}
-		else
+
+		// Then do the remapping if needed
+		TArray<UMaterialInterface*> StaticMeshMaterials;
+		TArray<FName> MaterialSlotNames;
+		UGeometryScriptLibrary_StaticMeshFunctions::GetMaterialListFromStaticMesh(StaticMesh, StaticMeshMaterials, MaterialSlotNames);
+
+		if (!StaticMeshMaterials.IsEmpty() && StaticMeshMaterials != OutDynMeshData->GetMaterials())
 		{
-			return true;
+			PCGGeometryHelpers::RemapMaterials(NewMesh, StaticMeshMaterials, OutDynMeshData->GetMutableMaterials());
 		}
+
+		return true;
 	};
 
 	switch (Settings->Mode)
@@ -257,6 +268,17 @@ bool FPCGAppendMeshesFromPointsElement::ExecuteInternal(FPCGContext* InContext)
 			return true;
 		}
 
+		// Remap materials if needed
+		const TArray<TObjectPtr<UMaterialInterface>>& InputMaterials = InAppendDynMeshData->GetMaterials();
+		TArray<TObjectPtr<UMaterialInterface>>& OutputMaterials = OutDynMeshData->GetMutableMaterials();
+		if (!InputMaterials.IsEmpty() && OutputMaterials != InputMaterials)
+		{
+			// If we have a remap to do, we will need to copy or steal the append mesh.
+			UPCGDynamicMeshData* InAppendDynamicMeshDataMutable = CopyOrSteal(InputAppendDynMesh[0], InContext);
+			PCGGeometryHelpers::RemapMaterials(InAppendDynamicMeshDataMutable->GetMutableDynamicMesh()->GetMeshRef(), InputMaterials, OutputMaterials);
+			InAppendDynMeshData = InAppendDynamicMeshDataMutable;
+		}
+
 		break;
 	}
 	default:
diff --git a/Engine/Plugins/PCGInterops/PCGGeometryScriptInterop/Source/PCGGeometryScriptInterop/Private/Elements/PCGDynamicMeshBaseElement.cpp b/Engine/Plugins/PCGInterops/PCGGeometryScriptInterop/Source/PCGGeometryScriptInterop/Private/Elements/PCGDynamicMeshBaseElement.cpp
--- a/Engine/Plugins/PCGInterops/PCGGeometryScriptInterop/Source/PCGGeometryScriptInterop/Private/Elements/PCGDynamicMeshBaseElement.cpp
+++ b/Engine/Plugins/PCGInterops/PCGGeometryScriptInterop/Source/PCGGeometryScriptInterop/Private/Elements/PCGDynamicMeshBaseElement.cpp
@@ -37,6 +37,8 @@ TArray<FPCGPinProperties> UPCGDynamicMeshBaseSettings::OutputPinProperties() con
 
 UPCGDynamicMeshData* IPCGDynamicMeshBaseElement::CopyOrSteal(const FPCGTaggedData& InTaggedData, FPCGContext* InContext)
 {
+	TRACE_CPUPROFILER_EVENT_SCOPE(IPCGDynamicMeshBaseElement::CopyOrSteal);
+	
 	check(InContext);
 	
 	const UPCGDynamicMeshData* InData = Cast<const UPCGDynamicMeshData>(InTaggedData.Data);
diff --git a/Engine/Plugins/PCGInterops/PCGGeometryScriptInterop/Source/PCGGeometryScriptInterop/Private/Elements/PCGGetDynamicMeshData.cpp b/Engine/Plugins/PCGInterops/PCGGeometryScriptInterop/Source/PCGGeometryScriptInterop/Private/Elements/PCGGetDynamicMeshData.cpp
--- a/Engine/Plugins/PCGInterops/PCGGeometryScriptInterop/Source/PCGGeometryScriptInterop/Private/Elements/PCGGetDynamicMeshData.cpp
+++ b/Engine/Plugins/PCGInterops/PCGGeometryScriptInterop/Source/PCGGeometryScriptInterop/Private/Elements/PCGGetDynamicMeshData.cpp
@@ -5,9 +5,12 @@
 #include "PCGContext.h"
 #include "Data/PCGDynamicMeshData.h"
 #include "Data/PCGGetDataFunctionRegistry.h"
+#include "Helpers/PCGGeometryHelpers.h"
 #include "Helpers/PCGHelpers.h"
 
 #include "DynamicMeshActor.h"
+#include "ConversionUtils/SceneComponentToDynamicMesh.h"
+#include "Materials/MaterialInterface.h"
 
 #define LOCTEXT_NAMESPACE "PCGGetDynamicMeshDataElement"
 
@@ -49,7 +52,8 @@ bool PCGGetDynamicMeshData::GetDynamicMeshDataFromActor(FPCGContext* InContext,
 	
 	if (ADynamicMeshActor* DynMeshActor = Cast<ADynamicMeshActor>(InActor))
 	{
-		if (!DynMeshActor->GetDynamicMeshComponent() || !DynMeshActor->GetDynamicMeshComponent()->GetDynamicMesh())
+		UDynamicMeshComponent* Component = DynMeshActor->GetDynamicMeshComponent();
+		if (!Component || !Component->GetDynamicMesh())
 		{
 			return false;
 		}
@@ -59,7 +63,7 @@ bool PCGGetDynamicMeshData::GetDynamicMeshDataFromActor(FPCGContext* InContext,
 		Algo::Transform(InActor->Tags, ActorTags, NameTagsToStringTags);
 		
 		UPCGDynamicMeshData* Data = FPCGContext::NewObject_AnyThread<UPCGDynamicMeshData>(InContext);
-		Data->Initialize(DynMeshActor->GetDynamicMeshComponent()->GetDynamicMesh(), InContext, /*bCanTakeOwnership=*/ false);
+		Data->Initialize(Component->GetDynamicMesh(), /*bCanTakeOwnership=*/false, Component->GetMaterials());
 
 		FPCGTaggedData& TaggedData = Output.Collection.TaggedData.Emplace_GetRef();
 		TaggedData.Data = Data;
@@ -95,15 +99,28 @@ bool PCGGetDynamicMeshData::GetDynamicMeshDataFromComponent(FPCGContext* InConte
 		check(NewDynamicMesh);
 
 		const UPCGGetDynamicMeshDataSettings* Settings = InContext ? InContext->GetInputSettings<UPCGGetDynamicMeshDataSettings>() : nullptr;
-		FGeometryScriptCopyMeshFromComponentOptions Options = Settings ? Settings->Options : FGeometryScriptCopyMeshFromComponentOptions{};
-
-		// TODO: Might want to store the transform on the data?
+		
+		// Adaptation of UGeometryScriptLibrary_SceneUtilityFunctions::CopyMeshFromComponent, since we don't have access to the Material list with this one.
 		FTransform Transform;
-		EGeometryScriptOutcomePins Outcome;
-		UGeometryScriptLibrary_SceneUtilityFunctions::CopyMeshFromComponent(SceneComponent, NewDynamicMesh, Options, /*bTransformToWorld=*/ false, Transform, Outcome);
+		FText ErrorMessage;
+		UE::Conversion::FToMeshOptions Options{};
+		if (Settings)
+		{
+			Options.bWantNormals = Settings->Options.bWantNormals;
+			Options.bWantTangents = Settings->Options.bWantTangents;
+			Options.bWantInstanceColors = Settings->Options.bWantInstanceColors;
+			Options.LODType = PCGGeometryHelpers::SafeConversionLODType(Settings->Options.RequestedLOD.LODType);
+			Options.LODIndex = Settings->Options.RequestedLOD.LODIndex;
+		}
 
-		if (Outcome == EGeometryScriptOutcomePins::Success)
+		TArray<UMaterialInterface*> ComponentMaterialList;
+		TArray<UMaterialInterface*> AssetMaterialList;
+
+		const bool bSuccess = UE::Conversion::SceneComponentToDynamicMesh(SceneComponent, Options, /*bTransformToWorld=*/false, NewDynamicMesh->GetMeshRef(), Transform, ErrorMessage, &ComponentMaterialList, &AssetMaterialList);
+		if (bSuccess)
 		{
+			OutputData->SetMaterials(!ComponentMaterialList.IsEmpty() ? ComponentMaterialList : AssetMaterialList);
+			
 			auto NameTagsToStringTags = [](const FName& InName) { return InName.ToString(); };
 
 			FPCGTaggedData& TaggedData = Output.Collection.TaggedData.Emplace_GetRef();
@@ -118,6 +135,11 @@ bool PCGGetDynamicMeshData::GetDynamicMeshDataFromComponent(FPCGContext* InConte
 				TaggedData.Tags.Append(ActorTags);
 			}
 		}
+		else
+		{
+			// If it fails, we still return true so we don't falloff on default getter behavior. 
+			PCGLog::LogErrorOnGraph(ErrorMessage, InContext);
+		}
 
 		return true;
 	}
diff --git a/Engine/Plugins/PCGInterops/PCGGeometryScriptInterop/Source/PCGGeometryScriptInterop/Private/Elements/PCGMergeDynamicMeshes.cpp b/Engine/Plugins/PCGInterops/PCGGeometryScriptInterop/Source/PCGGeometryScriptInterop/Private/Elements/PCGMergeDynamicMeshes.cpp
--- a/Engine/Plugins/PCGInterops/PCGGeometryScriptInterop/Source/PCGGeometryScriptInterop/Private/Elements/PCGMergeDynamicMeshes.cpp
+++ b/Engine/Plugins/PCGInterops/PCGGeometryScriptInterop/Source/PCGGeometryScriptInterop/Private/Elements/PCGMergeDynamicMeshes.cpp
@@ -8,6 +8,7 @@
 #include "DynamicMeshEditor.h"
 #include "UDynamicMesh.h"
 #include "GeometryScript/MeshBasicEditFunctions.h"
+#include "Helpers/PCGGeometryHelpers.h"
 
 #define LOCTEXT_NAMESPACE "PCGMergeDynamicMeshesElement"
 
@@ -68,8 +69,21 @@ bool FPCGMergeDynamicMeshesElement::ExecuteInternal(FPCGContext* InContext) cons
 		else
 		{
 			UE::Geometry::FMeshIndexMappings MeshIndexMappings;
-			UE::Geometry::FDynamicMeshEditor Editor(OutputData->GetMutableDynamicMesh()->GetMeshPtr());
-			Editor.AppendMesh(InputData->GetDynamicMesh()->GetMeshPtr(), MeshIndexMappings);
+			
+			{
+				TRACE_CPUPROFILER_EVENT_SCOPE(FPCGMergeDynamicMeshesElement::Execute::AppendMesh);
+				
+				UE::Geometry::FDynamicMeshEditor Editor(OutputData->GetMutableDynamicMesh()->GetMeshPtr());
+				Editor.AppendMesh(InputData->GetDynamicMesh()->GetMeshPtr(), MeshIndexMappings);
+			}
+
+			// It's also important to re-map the materials, but it's needed only if we do not have the same materials
+			const TArray<TObjectPtr<UMaterialInterface>>& InputMaterials = InputData->GetMaterials();
+			TArray<TObjectPtr<UMaterialInterface>>& OutputMaterials = OutputData->GetMutableMaterials();
+			if (!InputMaterials.IsEmpty() && InputMaterials != OutputMaterials)
+			{
+				PCGGeometryHelpers::RemapMaterials(OutputData->GetMutableDynamicMesh()->GetMeshRef(), InputMaterials, OutputMaterials, &MeshIndexMappings);
+			}
 		}
 	}
 	
diff --git a/Engine/Plugins/PCGInterops/PCGGeometryScriptInterop/Source/PCGGeometryScriptInterop/Private/Elements/PCGMeshSampler.cpp b/Engine/Plugins/PCGInterops/PCGGeometryScriptInterop/Source/PCGGeometryScriptInterop/Private/Elements/PCGMeshSampler.cpp
--- a/Engine/Plugins/PCGInterops/PCGGeometryScriptInterop/Source/PCGGeometryScriptInterop/Private/Elements/PCGMeshSampler.cpp
+++ b/Engine/Plugins/PCGInterops/PCGGeometryScriptInterop/Source/PCGGeometryScriptInterop/Private/Elements/PCGMeshSampler.cpp
@@ -403,30 +403,6 @@ bool FPCGMeshSamplerElement::PrepareDataInternal(FPCGContext* InContext) const
 
 	UGeometryScriptDebug* Debug = FPCGContext::NewObject_AnyThread<UGeometryScriptDebug>(Context);
 
-	// Make sure the LOD is valid, and make the conversion.
-	UE::Conversion::EMeshLODType RequestedLODType;
-
-	switch (Settings->RequestedLODType)
-	{
-	case EGeometryScriptLODType::MaxAvailable:
-		RequestedLODType = UE::Conversion::EMeshLODType::MaxAvailable;
-		break;
-	case EGeometryScriptLODType::HiResSourceModel:
-		RequestedLODType = UE::Conversion::EMeshLODType::HiResSourceModel;
-		break;
-	case EGeometryScriptLODType::SourceModel:
-		RequestedLODType = UE::Conversion::EMeshLODType::SourceModel;
-		break;
-	case EGeometryScriptLODType::RenderData:
-		RequestedLODType = UE::Conversion::EMeshLODType::RenderData;
-		break;
-	default:
-	{
-		PCGLog::LogErrorOnGraph(LOCTEXT("InvalidLOD", "Requested LOD is invalid."), Context);
-		return true;
-	}
-	}
-
 	for (const auto& [Path, DummyIndex, DummyIndex2] : Context->PathsToObjectsAndDataIndex)
 	{
 		UObject* Object = Path.ResolveObject();
@@ -454,7 +430,7 @@ bool FPCGMeshSamplerElement::PrepareDataInternal(FPCGContext* InContext) const
 			FText ErrorMessage;
 			UE::Conversion::FToMeshOptions Options{};
 			Options.bWantInstanceColors = true;
-			Options.LODType = RequestedLODType;
+			Options.LODType = PCGGeometryHelpers::SafeConversionLODType(Settings->RequestedLODType);
 			Options.LODIndex = Settings->RequestedLODIndex;
 
 			UE::Geometry::FDynamicMesh3 TempDynMesh{};
diff --git a/Engine/Plugins/PCGInterops/PCGGeometryScriptInterop/Source/PCGGeometryScriptInterop/Private/Elements/PCGSaveDynamicMeshToAsset.cpp b/Engine/Plugins/PCGInterops/PCGGeometryScriptInterop/Source/PCGGeometryScriptInterop/Private/Elements/PCGSaveDynamicMeshToAsset.cpp
--- a/Engine/Plugins/PCGInterops/PCGGeometryScriptInterop/Source/PCGGeometryScriptInterop/Private/Elements/PCGSaveDynamicMeshToAsset.cpp
+++ b/Engine/Plugins/PCGInterops/PCGGeometryScriptInterop/Source/PCGGeometryScriptInterop/Private/Elements/PCGSaveDynamicMeshToAsset.cpp
@@ -88,6 +88,13 @@ bool FPCGSaveDynamicMeshToAssetElement::ExecuteInternal(FPCGContext* InContext)
 		AssetOptions.bEmitTransaction = false;
 		AssetOptions.bDeferMeshPostEditChange = false;
 
+		// Force the option to replace materials if the dynamic mesh data has materials
+		if (!AssetOptions.bReplaceMaterials && !DynamicMeshData->GetMaterials().IsEmpty())
+		{
+			AssetOptions.bReplaceMaterials = true;
+			AssetOptions.NewMaterials = DynamicMeshData->GetMaterials();
+		}
+
 		// Convert merged mesh to static mesh, const_cast because GeometryScript API is not nice with constness, but it will not be modified.
 		EGeometryScriptOutcomePins OutResult;
 		UGeometryScriptLibrary_StaticMeshFunctions::CopyMeshToStaticMesh(const_cast<UDynamicMesh*>(DynamicMeshData->GetDynamicMesh()), CastChecked<UStaticMesh>(Asset), AssetOptions, Settings->MeshWriteLOD, OutResult);
diff --git a/Engine/Plugins/PCGInterops/PCGGeometryScriptInterop/Source/PCGGeometryScriptInterop/Private/Elements/PCGSpawnDynamicMesh.cpp b/Engine/Plugins/PCGInterops/PCGGeometryScriptInterop/Source/PCGGeometryScriptInterop/Private/Elements/PCGSpawnDynamicMesh.cpp
--- a/Engine/Plugins/PCGInterops/PCGGeometryScriptInterop/Source/PCGGeometryScriptInterop/Private/Elements/PCGSpawnDynamicMesh.cpp
+++ b/Engine/Plugins/PCGInterops/PCGGeometryScriptInterop/Source/PCGGeometryScriptInterop/Private/Elements/PCGSpawnDynamicMesh.cpp
@@ -67,7 +67,7 @@ bool FPCGSpawnDynamicMeshElement::ExecuteInternal(FPCGContext* InContext) const
 		}
 
 		// Needs to be a copy.
-		Component->SetMesh(UE::Geometry::FDynamicMesh3(DynMeshData->GetDynamicMesh()->GetMeshRef()));
+		DynMeshData->InitializeDynamicMeshComponentFromData(Component);
 
 		for (const FString& Tag : Input.Tags)
 		{
diff --git a/Engine/Plugins/PCGInterops/PCGGeometryScriptInterop/Source/PCGGeometryScriptInterop/Private/Elements/PCGStaticMeshToDynamicMeshElement.cpp b/Engine/Plugins/PCGInterops/PCGGeometryScriptInterop/Source/PCGGeometryScriptInterop/Private/Elements/PCGStaticMeshToDynamicMeshElement.cpp
--- a/Engine/Plugins/PCGInterops/PCGGeometryScriptInterop/Source/PCGGeometryScriptInterop/Private/Elements/PCGStaticMeshToDynamicMeshElement.cpp
+++ b/Engine/Plugins/PCGInterops/PCGGeometryScriptInterop/Source/PCGGeometryScriptInterop/Private/Elements/PCGStaticMeshToDynamicMeshElement.cpp
@@ -9,6 +9,7 @@
 #include "UDynamicMesh.h"
 #include "ConversionUtils/SceneComponentToDynamicMesh.h"
 #include "Engine/StaticMesh.h"
+#include "GeometryScript/MeshAssetFunctions.h"
 
 #define LOCTEXT_NAMESPACE "PCGStaticMeshToDynamicMeshElementElement"
 
@@ -71,8 +72,16 @@ bool FPCGStaticMeshToDynamicMeshElement::PrepareDataInternal(FPCGContext* InCont
 	{
 		return true;
 	}
+
+	TArray<FSoftObjectPath> ObjectsToLoad;
+	if (Settings->bExtractMaterials)
+	{
+		Algo::Transform(Settings->OverrideMaterials, ObjectsToLoad, [](const TSoftObjectPtr<UMaterialInterface>& MaterialSoftPtr) { return MaterialSoftPtr.ToSoftObjectPath(); });
+	}
+
+	ObjectsToLoad.Add(Settings->StaticMesh.ToSoftObjectPath());
 	
-	return Context->RequestResourceLoad(Context, {Settings->StaticMesh.ToSoftObjectPath()}, !Settings->bSynchronousLoad);
+	return Context->RequestResourceLoad(Context, std::move(ObjectsToLoad), !Settings->bSynchronousLoad);
 }
 
 bool FPCGStaticMeshToDynamicMeshElement::ExecuteInternal(FPCGContext* InContext) const
@@ -94,6 +103,20 @@ bool FPCGStaticMeshToDynamicMeshElement::ExecuteInternal(FPCGContext* InContext)
 		
 		return true;
 	}
+
+	if (Settings->bExtractMaterials && !Settings->OverrideMaterials.IsEmpty())
+	{
+		if (StaticMesh->GetStaticMaterials().Num() != Settings->OverrideMaterials.Num())
+		{
+			PCGLog::LogErrorOnGraph(FText::Format(LOCTEXT("MismatchMaterials", "Mismatch number between Static mesh materials ({0}) and override materials ({1})"), StaticMesh->GetStaticMaterials().Num(), Settings->OverrideMaterials.Num()));
+			return true;
+		}
+		else if (const TSoftObjectPtr<UMaterialInterface>* UnloadedMaterial = Settings->OverrideMaterials.FindByPredicate([](const TSoftObjectPtr<UMaterialInterface>& MaterialSoftPtr) -> bool { return !MaterialSoftPtr.Get(); }))
+		{
+			PCGLog::LogErrorOnGraph(FText::Format(LOCTEXT("UnloadedMaterial", "Material {0} failed to load."), FText::FromString(UnloadedMaterial->ToSoftObjectPath().ToString())));
+			return true;
+		}
+	}
 	
 	UE::Conversion::EMeshLODType LODType = static_cast<UE::Conversion::EMeshLODType>(Settings->RequestedLODType);
 	int32 LODIndex = Settings->RequestedLODIndex;
@@ -105,8 +128,24 @@ bool FPCGStaticMeshToDynamicMeshElement::ExecuteInternal(FPCGContext* InContext)
 
 	if (bSuccess)
 	{
+		TArray<UMaterialInterface*> Materials;
+
+		if (Settings->bExtractMaterials)
+		{
+			if (!Settings->OverrideMaterials.IsEmpty())
+			{
+				Algo::Transform(Settings->OverrideMaterials, Materials, [](const TSoftObjectPtr<UMaterialInterface>& MaterialSoftPtr) { return MaterialSoftPtr.Get(); });
+			}
+			else
+			{
+				TArray<FName> MaterialSlotNames;
+				UGeometryScriptLibrary_StaticMeshFunctions::GetMaterialListFromStaticMesh(StaticMesh, Materials, MaterialSlotNames);
+			}
+		}
+		
 		UPCGDynamicMeshData* DynMeshData = FPCGContext::NewObject_AnyThread<UPCGDynamicMeshData>(InContext);
-		DynMeshData->Initialize(std::move(NewMesh), InContext);
+		DynMeshData->Initialize(std::move(NewMesh), Materials);
+
 		InContext->OutputData.TaggedData.Emplace_GetRef().Data = DynMeshData;
 	}
 	else
diff --git a/Engine/Plugins/PCGInterops/PCGGeometryScriptInterop/Source/PCGGeometryScriptInterop/Private/Helpers/PCGGeometryHelpers.cpp b/Engine/Plugins/PCGInterops/PCGGeometryScriptInterop/Source/PCGGeometryScriptInterop/Private/Helpers/PCGGeometryHelpers.cpp
--- a/Engine/Plugins/PCGInterops/PCGGeometryScriptInterop/Source/PCGGeometryScriptInterop/Private/Helpers/PCGGeometryHelpers.cpp
+++ b/Engine/Plugins/PCGInterops/PCGGeometryScriptInterop/Source/PCGGeometryScriptInterop/Private/Helpers/PCGGeometryHelpers.cpp
@@ -2,10 +2,10 @@
 
 #include "Helpers/PCGGeometryHelpers.h"
 
+#include "ConversionUtils/SceneComponentToDynamicMesh.h"
+#include "DynamicMesh/DynamicMeshAttributeSet.h"
 #include "Utils/PCGLogErrors.h"
 
-#include "GeometryScript/GeometryScriptTypes.h"
-
 namespace PCGGeometryHelpers
 {
 	void GeometryScriptDebugToPCGLog(FPCGContext* Context, const UGeometryScriptDebug* Debug)
@@ -24,4 +24,87 @@ namespace PCGGeometryHelpers
 			}
 		}
 	}
+
+	UE::Conversion::EMeshLODType SafeConversionLODType(const EGeometryScriptLODType LODType)
+	{
+		// Make sure the LOD is valid, and make the conversion.
+		UE::Conversion::EMeshLODType RequestedLODType{};
+
+		switch (LODType)
+		{
+		case EGeometryScriptLODType::MaxAvailable:
+			RequestedLODType = UE::Conversion::EMeshLODType::MaxAvailable;
+			break;
+		case EGeometryScriptLODType::HiResSourceModel:
+			RequestedLODType = UE::Conversion::EMeshLODType::HiResSourceModel;
+			break;
+		case EGeometryScriptLODType::SourceModel:
+			RequestedLODType = UE::Conversion::EMeshLODType::SourceModel;
+			break;
+		case EGeometryScriptLODType::RenderData:
+			RequestedLODType = UE::Conversion::EMeshLODType::RenderData;
+			break;
+		default:
+			break;
+		}
+
+		return RequestedLODType;
+	}
+
+	void RemapMaterials(UE::Geometry::FDynamicMesh3& InMesh, const TArray<UMaterialInterface*>& FromMaterials, TArray<TObjectPtr<UMaterialInterface>>& ToMaterials, const UE::Geometry::FMeshIndexMappings* OptionalMappings)
+	{
+		if (FromMaterials.IsEmpty() || ToMaterials.IsEmpty() || !InMesh.HasAttributes() || !InMesh.Attributes()->HasMaterialID())
+		{
+			return;
+		}
+		
+		TMap<int32, int32> MaterialIDRemap;
+		MaterialIDRemap.Reserve(FromMaterials.Num());
+		
+		for (int32 FromMaterialIndex = 0; FromMaterialIndex < FromMaterials.Num(); ++FromMaterialIndex)
+		{
+			UMaterialInterface* FromMaterial = FromMaterials[FromMaterialIndex];
+			int32 ToMaterialIndex = ToMaterials.IndexOfByKey(FromMaterial);
+			if (ToMaterialIndex == INDEX_NONE)
+			{
+				ToMaterialIndex = ToMaterials.Add(FromMaterial);
+			}
+
+			if (ToMaterialIndex != FromMaterialIndex)
+			{
+				MaterialIDRemap.Emplace(FromMaterialIndex, ToMaterialIndex);
+			}
+		}
+
+		if (!MaterialIDRemap.IsEmpty())
+		{
+			TRACE_CPUPROFILER_EVENT_SCOPE(PCGGeometryHelpers::RemapMaterials);
+			UE::Geometry::FDynamicMeshMaterialAttribute* MaterialAttribute = InMesh.Attributes()->GetMaterialID();
+
+			auto Remap = [&MaterialAttribute, &MaterialIDRemap](const int32 TriangleID)
+			{
+				const int32 OriginalMaterialID = MaterialAttribute->GetValue(TriangleID);
+				if (const int32* RemappedMaterialID = MaterialIDRemap.Find(OriginalMaterialID))
+				{
+					MaterialAttribute->SetValue(TriangleID, *RemappedMaterialID);
+				}
+			};
+			
+			// TODO: Could be parallelized
+			if (OptionalMappings)
+			{
+				for (const TPair<int32, int32>& MapTriangleID : OptionalMappings->GetTriangleMap().GetForwardMap())
+				{
+					Remap(MapTriangleID.Value);
+				}
+			}
+			else
+			{
+				for (const int32 TriangleID : InMesh.TriangleIndicesItr())
+				{
+					Remap(TriangleID);
+				}
+			}
+		}
+	}
 }
\ No newline at end of file
diff --git a/Engine/Plugins/PCGInterops/PCGGeometryScriptInterop/Source/PCGGeometryScriptInterop/Public/Data/PCGDynamicMeshData.h b/Engine/Plugins/PCGInterops/PCGGeometryScriptInterop/Source/PCGGeometryScriptInterop/Public/Data/PCGDynamicMeshData.h
--- a/Engine/Plugins/PCGInterops/PCGGeometryScriptInterop/Source/PCGGeometryScriptInterop/Public/Data/PCGDynamicMeshData.h
+++ b/Engine/Plugins/PCGInterops/PCGGeometryScriptInterop/Source/PCGGeometryScriptInterop/Public/Data/PCGDynamicMeshData.h
@@ -4,8 +4,8 @@
 
 #include "Data/PCGSpatialData.h"
 
-#include "DynamicMesh/DynamicMesh3.h"
 #include "DynamicMesh/DynamicMeshOctree3.h"
+#include "DynamicMesh/MeshIndexMappings.h"
 
 #include "Misc/SpinLock.h"
 
@@ -13,6 +13,8 @@
 
 struct FPCGContext;
 class UDynamicMesh;
+class UDynamicMeshComponent;
+class UMaterialInterface;
 
 UCLASS(BlueprintType, ClassGroup = (Procedural))
 class PCGGEOMETRYSCRIPTINTEROP_API UPCGDynamicMeshData : public UPCGSpatialData
@@ -22,11 +24,17 @@ class PCGGEOMETRYSCRIPTINTEROP_API UPCGDynamicMeshData : public UPCGSpatialData
 public:
 	UPCGDynamicMeshData(const FObjectInitializer& ObjectInitializer);
 	
-	void Initialize(UDynamicMesh* InMesh, FPCGContext* Context = nullptr, bool bCanTakeOwnership = false);
-	void Initialize(UE::Geometry::FDynamicMesh3&& InMesh, FPCGContext* Context = nullptr);
-	
-	UFUNCTION(BlueprintCallable, Category="DynamicMesh", meta = (DisplayName = "Initialize"))
-	void K2_Initialize(UDynamicMesh* InMesh, FPCGContext& Context) { Initialize(InMesh, &Context); }
+	void Initialize(UDynamicMesh* InMesh, bool bCanTakeOwnership = false, const TArray<UMaterialInterface*>& InOptionalMaterials = {});
+	void Initialize(UE::Geometry::FDynamicMesh3&& InMesh, const TArray<UMaterialInterface*>& InOptionalMaterials = {});
+
+	/**
+	 * Initialize the dynamic mesh data from an input dynamic mesh object.
+	 * If the input dynamic mesh is not meant to be re-used after this initialization, you can set Can Take Ownership to true. Be careful as it
+	 * will put the previous object in an invalid state.
+	 * You can also pass an array of materials that correspond to the referenced materials in the dynamic mesh.
+	 */
+	UFUNCTION(BlueprintCallable, Category="DynamicMesh", meta = (DisplayName = "Initialize", AutoCreateRefTerm = "InMaterials"))
+	void K2_Initialize(UDynamicMesh* InMesh, const TArray<UMaterialInterface*>& InMaterials, bool bCanTakeOwnership = false) { Initialize(InMesh, bCanTakeOwnership, InMaterials); }
 	
 	// ~Begin UPCGData interface
 	virtual EPCGDataType GetDataType() const override { return EPCGDataType::DynamicMesh; }
@@ -46,6 +54,15 @@ public:
 	UDynamicMesh* GetMutableDynamicMesh() { bDynamicMeshBoundsAreDirty = true; bDynamicMeshOctreeIsDirty = true; return DynamicMesh; }
 	const UDynamicMesh* GetDynamicMesh() const { return DynamicMesh; }
 
+	UFUNCTION(BlueprintCallable, Category="DynamicMesh")
+	void SetMaterials(const TArray<UMaterialInterface*>& InMaterials);
+	
+	TArray<TObjectPtr<UMaterialInterface>>& GetMutableMaterials() { return Materials; }
+	const TArray<TObjectPtr<UMaterialInterface>>& GetMaterials() const { return Materials; }
+
+	// Copy the mesh of the data into the component and set the materials.
+	void InitializeDynamicMeshComponentFromData(UDynamicMeshComponent* InComponent) const;
+
 protected:
 	// ~Begin UPCGSpatialData interface
 	virtual UPCGSpatialData* CopyInternal(FPCGContext* Context) const override;
@@ -60,6 +77,9 @@ private:
 protected:
 	UPROPERTY(BlueprintReadOnly, VisibleAnywhere, Instanced, Category="DynamicMesh")
 	TObjectPtr<UDynamicMesh> DynamicMesh;
+	
+	UPROPERTY(BlueprintReadOnly, VisibleAnywhere, Instanced, Category="DynamicMesh")
+	TArray<TObjectPtr<UMaterialInterface>> Materials;
 
 	mutable UE::Geometry::FDynamicMeshOctree3 DynamicMeshOctree;
 	mutable bool bDynamicMeshOctreeIsDirty = true;
diff --git a/Engine/Plugins/PCGInterops/PCGGeometryScriptInterop/Source/PCGGeometryScriptInterop/Public/Elements/PCGStaticMeshToDynamicMeshElement.h b/Engine/Plugins/PCGInterops/PCGGeometryScriptInterop/Source/PCGGeometryScriptInterop/Public/Elements/PCGStaticMeshToDynamicMeshElement.h
--- a/Engine/Plugins/PCGInterops/PCGGeometryScriptInterop/Source/PCGGeometryScriptInterop/Public/Elements/PCGStaticMeshToDynamicMeshElement.h
+++ b/Engine/Plugins/PCGInterops/PCGGeometryScriptInterop/Source/PCGGeometryScriptInterop/Public/Elements/PCGStaticMeshToDynamicMeshElement.h
@@ -38,11 +38,19 @@ public:
 	UPROPERTY(BlueprintReadWrite, EditAnywhere, Category="Settings", meta = (PCG_Overridable))
 	TSoftObjectPtr<UStaticMesh> StaticMesh;
 
+	/** Allows to extract materials from the static mesh and store them in the PCG Dynamic Mesh Data. */
+	UPROPERTY(BlueprintReadWrite, EditAnywhere, Category="Settings", meta = (PCG_Overridable))
+	bool bExtractMaterials = true;
+
+	/** If it extracts materials, we can specify override materials. It needs to have the same number of material overrides than there are materials on the static mesh. */
+	UPROPERTY(EditAnywhere, Category = "Settings", meta = (EditCondition = "bExtractMaterials", EditConditionHides))
+	TArray<TSoftObjectPtr<UMaterialInterface>> OverrideMaterials;
+
 	/** LOD type to use when creating DynamicMesh from specified StaticMesh. */
-	UPROPERTY(BlueprintReadWrite, EditAnywhere, Category = "Settings|LODSettings", meta = (EditCondition = "Mode!=EPCGAppendMeshesFromPointsMode::DynamicMesh", EditConditionHides, PCG_Overridable))
-	EGeometryScriptLODType RequestedLODType = EGeometryScriptLODType::RenderData;
+	UPROPERTY(BlueprintReadWrite, EditAnywhere, Category = "Settings|LODSettings", meta = (PCG_Overridable))
+	EGeometryScriptLODType RequestedLODType = EGeometryScriptLODType::MaxAvailable;
 
-	UPROPERTY(BlueprintReadWrite, EditAnywhere, Category = "Settings|LODSettings", meta = (EditCondition = "Mode!=EPCGAppendMeshesFromPointsMode::DynamicMesh", EditConditionHides, PCG_Overridable))
+	UPROPERTY(BlueprintReadWrite, EditAnywhere, Category = "Settings|LODSettings", meta = (PCG_Overridable))
 	int32 RequestedLODIndex = 0;
 	
 	UPROPERTY(BlueprintReadWrite, EditAnywhere, Category = "Settings|Debug")
diff --git a/Engine/Plugins/PCGInterops/PCGGeometryScriptInterop/Source/PCGGeometryScriptInterop/Public/Helpers/PCGGeometryHelpers.h b/Engine/Plugins/PCGInterops/PCGGeometryScriptInterop/Source/PCGGeometryScriptInterop/Public/Helpers/PCGGeometryHelpers.h
--- a/Engine/Plugins/PCGInterops/PCGGeometryScriptInterop/Source/PCGGeometryScriptInterop/Public/Helpers/PCGGeometryHelpers.h
+++ b/Engine/Plugins/PCGInterops/PCGGeometryScriptInterop/Source/PCGGeometryScriptInterop/Public/Helpers/PCGGeometryHelpers.h
@@ -4,10 +4,30 @@
 
 #include "Elements/PCGMeshSampler.h"
 
+#include "ConversionUtils/SceneComponentToDynamicMesh.h"
+#include "GeometryScript/GeometryScriptTypes.h"
+
 struct FPCGContext;
 class UGeometryScriptDebug;
+class UMaterialInterface;
+
+namespace UE::Geometry
+{
+	class FDynamicMesh3;
+	struct FMeshIndexMappings;
+}
 
 namespace PCGGeometryHelpers
 {
 	PCGGEOMETRYSCRIPTINTEROP_API void GeometryScriptDebugToPCGLog(FPCGContext* Context, const UGeometryScriptDebug* Debug);
+	UE::Conversion::EMeshLODType SafeConversionLODType(const EGeometryScriptLODType LODType);
+
+	/**
+	 * Adaptation of UGeometryScriptLibrary_MeshMaterialFunctions::RemapToNewMaterialIDsByMaterial, to work on UE::Geometry::FDynamicMesh3 and with optional mappings.
+	 * @param InMesh: Mesh to modify
+	 * @param FromMaterials: Original array of materials for the mesh
+	 * @param ToMaterials: New array of materials for the mesh. Mutable. If the original material is not in the new array, add it.
+	 * @param OptionalMappings: Optional mappings if we want to update a subset of the triangles in the mesh.
+	 */
+	void RemapMaterials(UE::Geometry::FDynamicMesh3& InMesh, const TArray<UMaterialInterface*>& FromMaterials, TArray<TObjectPtr<UMaterialInterface>>& ToMaterials, const UE::Geometry::FMeshIndexMappings* OptionalMappings = nullptr);
 }
