[PCG] Add more option on how materials are extracted/replacing existing materials

#rnx

diff --git a/Engine/Plugins/PCGInterops/PCGGeometryScriptInterop/Source/PCGGeometryScriptInterop/Private/Elements/PCGAppendMeshesFromPoints.cpp b/Engine/Plugins/PCGInterops/PCGGeometryScriptInterop/Source/PCGGeometryScriptInterop/Private/Elements/PCGAppendMeshesFromPoints.cpp
--- a/Engine/Plugins/PCGInterops/PCGGeometryScriptInterop/Source/PCGGeometryScriptInterop/Private/Elements/PCGAppendMeshesFromPoints.cpp
+++ b/Engine/Plugins/PCGInterops/PCGGeometryScriptInterop/Source/PCGGeometryScriptInterop/Private/Elements/PCGAppendMeshesFromPoints.cpp
@@ -223,13 +223,16 @@ bool FPCGAppendMeshesFromPointsElement::ExecuteInternal(FPCGContext* InContext)
 		}
 
 		// Then do the remapping if needed
-		TArray<UMaterialInterface*> StaticMeshMaterials;
-		TArray<FName> MaterialSlotNames;
-		UGeometryScriptLibrary_StaticMeshFunctions::GetMaterialListFromStaticMesh(StaticMesh, StaticMeshMaterials, MaterialSlotNames);
-
-		if (!StaticMeshMaterials.IsEmpty() && StaticMeshMaterials != OutDynMeshData->GetMaterials())
+		if (Settings->bExtractMaterials)
 		{
-			PCGGeometryHelpers::RemapMaterials(NewMesh, StaticMeshMaterials, OutDynMeshData->GetMutableMaterials());
+			TArray<UMaterialInterface*> StaticMeshMaterials;
+			TArray<FName> MaterialSlotNames;
+			UGeometryScriptLibrary_StaticMeshFunctions::GetMaterialListFromStaticMesh(StaticMesh, StaticMeshMaterials, MaterialSlotNames);
+
+			if (!StaticMeshMaterials.IsEmpty() && StaticMeshMaterials != OutDynMeshData->GetMaterials())
+			{
+				PCGGeometryHelpers::RemapMaterials(NewMesh, StaticMeshMaterials, OutDynMeshData->GetMutableMaterials());
+			}
 		}
 
 		return true;
diff --git a/Engine/Plugins/PCGInterops/PCGGeometryScriptInterop/Source/PCGGeometryScriptInterop/Private/Elements/PCGSaveDynamicMeshToAsset.cpp b/Engine/Plugins/PCGInterops/PCGGeometryScriptInterop/Source/PCGGeometryScriptInterop/Private/Elements/PCGSaveDynamicMeshToAsset.cpp
--- a/Engine/Plugins/PCGInterops/PCGGeometryScriptInterop/Source/PCGGeometryScriptInterop/Private/Elements/PCGSaveDynamicMeshToAsset.cpp
+++ b/Engine/Plugins/PCGInterops/PCGGeometryScriptInterop/Source/PCGGeometryScriptInterop/Private/Elements/PCGSaveDynamicMeshToAsset.cpp
@@ -89,10 +89,13 @@ bool FPCGSaveDynamicMeshToAssetElement::ExecuteInternal(FPCGContext* InContext)
 		AssetOptions.bDeferMeshPostEditChange = false;
 
 		// Force the option to replace materials if the dynamic mesh data has materials
-		if (!AssetOptions.bReplaceMaterials && !DynamicMeshData->GetMaterials().IsEmpty())
+		if (Settings->bExportMaterialsFromDynamicMesh)
 		{
-			AssetOptions.bReplaceMaterials = true;
-			AssetOptions.NewMaterials = DynamicMeshData->GetMaterials();
+			AssetOptions.bReplaceMaterials = !DynamicMeshData->GetMaterials().IsEmpty();
+			if (!DynamicMeshData->GetMaterials().IsEmpty())
+			{
+				AssetOptions.NewMaterials = DynamicMeshData->GetMaterials();
+			}
 		}
 
 		// Convert merged mesh to static mesh, const_cast because GeometryScript API is not nice with constness, but it will not be modified.
diff --git a/Engine/Plugins/PCGInterops/PCGGeometryScriptInterop/Source/PCGGeometryScriptInterop/Public/Elements/PCGAppendMeshesFromPoints.h b/Engine/Plugins/PCGInterops/PCGGeometryScriptInterop/Source/PCGGeometryScriptInterop/Public/Elements/PCGAppendMeshesFromPoints.h
--- a/Engine/Plugins/PCGInterops/PCGGeometryScriptInterop/Source/PCGGeometryScriptInterop/Public/Elements/PCGAppendMeshesFromPoints.h
+++ b/Engine/Plugins/PCGInterops/PCGGeometryScriptInterop/Source/PCGGeometryScriptInterop/Public/Elements/PCGAppendMeshesFromPoints.h
@@ -53,6 +53,10 @@ public:
 	UPROPERTY(BlueprintReadWrite, EditAnywhere, Category="Settings", meta = (EditCondition = "Mode==EPCGAppendMeshesFromPointsMode::StaticMeshFromAttribute", EditConditionHides, PCG_Overridable))
 	FPCGAttributePropertyInputSelector MeshAttribute;
 	
+	/** Allows to extract materials from the static mesh and set them in the resulting append. */
+    UPROPERTY(BlueprintReadWrite, EditAnywhere, Category="Settings", meta = (EditCondition = "Mode!=EPCGAppendMeshesFromPointsMode::DynamicMesh", EditConditionHides, PCG_Overridable))
+    bool bExtractMaterials = true;
+	
 	/** LOD type to use when creating DynamicMesh from specified StaticMesh. */
 	UPROPERTY(BlueprintReadWrite, EditAnywhere, Category = "Settings|LODSettings", meta = (EditCondition = "Mode!=EPCGAppendMeshesFromPointsMode::DynamicMesh", EditConditionHides, PCG_Overridable))
 	EGeometryScriptLODType RequestedLODType = EGeometryScriptLODType::RenderData;
diff --git a/Engine/Plugins/PCGInterops/PCGGeometryScriptInterop/Source/PCGGeometryScriptInterop/Public/Elements/PCGSaveDynamicMeshToAsset.h b/Engine/Plugins/PCGInterops/PCGGeometryScriptInterop/Source/PCGGeometryScriptInterop/Public/Elements/PCGSaveDynamicMeshToAsset.h
--- a/Engine/Plugins/PCGInterops/PCGGeometryScriptInterop/Source/PCGGeometryScriptInterop/Public/Elements/PCGSaveDynamicMeshToAsset.h
+++ b/Engine/Plugins/PCGInterops/PCGGeometryScriptInterop/Source/PCGGeometryScriptInterop/Public/Elements/PCGSaveDynamicMeshToAsset.h
@@ -35,6 +35,14 @@ public:
 	UPROPERTY(BlueprintReadWrite, EditAnywhere, Category = Settings, meta = (ShowOnlyInnerProperties, PCG_Overridable))
 	FPCGAssetExporterParameters ExportParams;
 
+	/**
+	 * This option has higher priority than CopyMeshToAssetOptions.ReplaceMaterials.
+	 * If true, we will replace the materials from the materials stored on the PCG Dynamic Mesh data.
+	 * Otherwise, we will follow what is set in CopyMeshToAssetOptions.
+	 */
+	UPROPERTY(BlueprintReadWrite, EditAnywhere, Category = Settings, meta = (PCG_Overridable))
+	bool bExportMaterialsFromDynamicMesh = true;
+
 	UPROPERTY(BlueprintReadWrite, EditAnywhere, Category = Settings, meta = (PCG_Overridable))
 	FGeometryScriptCopyMeshToAssetOptions CopyMeshToAssetOptions;
 	
