[PCG] Make a simple PCG Geometry Blueprint Element to streamline the usage of copy and steal + simple processing function.

#rnx

diff --git a/Engine/Plugins/PCG/Source/PCG/Public/Elements/PCGExecuteBlueprint.h b/Engine/Plugins/PCG/Source/PCG/Public/Elements/PCGExecuteBlueprint.h
--- a/Engine/Plugins/PCG/Source/PCG/Public/Elements/PCGExecuteBlueprint.h
+++ b/Engine/Plugins/PCG/Source/PCG/Public/Elements/PCGExecuteBlueprint.h
@@ -414,6 +414,8 @@ public:
 	virtual bool CanExecuteOnlyOnMainThread(FPCGContext* Context) const override;
 	virtual bool IsCacheable(const UPCGSettings* InSettings) const override;
 	virtual bool ShouldComputeFullOutputDataCrc(FPCGContext* Context) const override;
+	/** Set it to true by default, if there is a performance concern, we can expose a bool in the element class. */
+	virtual bool ShouldVerifyIfOutputsAreUsedMultipleTimes(const UPCGSettings* InSettings) const override { return true; }
 
 protected:
 	virtual bool ExecuteInternal(FPCGContext* Context) const override;
diff --git a/Engine/Plugins/PCGInterops/PCGGeometryScriptInterop/Source/PCGGeometryScriptInterop/Private/Elements/PCGGeometryBlueprintElement.cpp b/Engine/Plugins/PCGInterops/PCGGeometryScriptInterop/Source/PCGGeometryScriptInterop/Private/Elements/PCGGeometryBlueprintElement.cpp
new file mode 100644
--- /dev/null
+++ b/Engine/Plugins/PCGInterops/PCGGeometryScriptInterop/Source/PCGGeometryScriptInterop/Private/Elements/PCGGeometryBlueprintElement.cpp
@@ -0,0 +1,94 @@
+// Copyright Epic Games, Inc. All Rights Reserved.
+
+#include "Elements/PCGGeometryBlueprintElement.h"
+
+#include "Data/PCGDynamicMeshData.h"
+#include "Elements/PCGDynamicMeshBaseElement.h"
+
+#include "EdGraph/EdGraph.h"
+#include "Engine/Blueprint.h"
+
+#define LOCTEXT_NAMESPACE "PCGGeometryBlueprintElement"
+
+UPCGGeometryBlueprintElement::UPCGGeometryBlueprintElement()
+	: UPCGBlueprintElement()
+{
+	// Setup the element for basic dynamic mesh processing
+	bIsCacheable = false;
+	bHasDefaultInPin = false;
+	bHasDefaultOutPin = false;
+
+	CustomInputPins.Emplace_GetRef(PCGPinConstants::DefaultInputLabel, EPCGDataType::DynamicMesh).SetRequiredPin();
+	CustomOutputPins.Emplace(PCGPinConstants::DefaultOutputLabel, EPCGDataType::DynamicMesh);
+}
+
+void UPCGGeometryBlueprintElement::ExecuteWithContext_Implementation(FPCGContext& InContext, const FPCGDataCollection& Input, FPCGDataCollection& Output)
+{
+	// Verify that we are in the right setup
+	const bool bHasASingleInputPin = (bHasDefaultInPin && CustomInputPins.IsEmpty()) || (!bHasDefaultInPin && CustomInputPins.Num() == 1);
+	const bool bHasASingleOutputPin = (bHasDefaultOutPin && CustomOutputPins.IsEmpty()) || (!bHasDefaultOutPin && CustomOutputPins.Num() == 1);
+
+#if WITH_EDITOR
+	const UBlueprint* Blueprint = Cast<UBlueprint>(GetClass()->ClassGeneratedBy);
+	auto IsFunctionOverridden = [Blueprint](const FName FunctionName)
+	{
+		return Blueprint && Blueprint->FunctionGraphs.ContainsByPredicate([FunctionName](const TObjectPtr<UEdGraph>& Graph) { return Graph->GetFName() == FunctionName;});
+	};
+	
+	const bool bIsExecuteOverridden = IsFunctionOverridden(GET_FUNCTION_NAME_CHECKED(UPCGBlueprintElement, Execute));
+	const bool bIsProcessDynMeshOverridden = IsFunctionOverridden(GET_FUNCTION_NAME_CHECKED(UPCGGeometryBlueprintElement, ProcessDynamicMesh));
+#else
+	// We can't know in runtime build as the functions are compiled.
+	const bool bIsExecuteOverridden = false;
+	const bool bIsProcessDynMeshOverridden = false;
+#endif // WITH_EDITOR
+	
+	if (!bHasASingleInputPin || !bHasASingleOutputPin || bIsExecuteOverridden)
+	{
+		// Make sure to throw a warning if we are in this case and Process Dynamic Mesh is overridden to warn the user their function won't be called
+		if (bIsProcessDynMeshOverridden)
+		{
+			PCGLog::LogWarningOnGraph(LOCTEXT("DynMeshOverridenButNotCalled", "Process Dynamic Mesh was overridden, but we don't have the expected setup (single input and output pin)"
+				" or Execute is also overriden. Process Dynamic Mesh won't be called."), &GetContext());
+		}
+		
+		// If we have not the right number of pins, or Execute is overriden, just call the parent function.
+		UPCGBlueprintElement::ExecuteWithContext_Implementation(InContext, Input, Output);
+		return;
+	}
+
+	const FName InputPinLabel = bHasDefaultInPin ? PCGPinConstants::DefaultInputLabel : CustomInputPins[0].Label;
+	const FName OutputPinLabel = bHasDefaultOutPin ? PCGPinConstants::DefaultOutputLabel : CustomOutputPins[0].Label;
+	
+	for (const FPCGTaggedData& InputData : Input.GetInputsByPin(InputPinLabel))
+	{
+		if (!Cast<const UPCGDynamicMeshData>(InputData.Data))
+		{
+			continue;
+		}
+
+		UPCGDynamicMeshData* ProcessingMesh = CopyOrStealInputData(InputData);
+		TArray<FString> OutTags;
+
+		ProcessDynamicMesh(ProcessingMesh->GetMutableDynamicMesh(), OutTags);
+		FPCGTaggedData& OutputData = Output.TaggedData.Emplace_GetRef(InputData);
+		OutputData.Tags.Append(OutTags);
+		OutputData.Data = ProcessingMesh;
+		OutputData.Pin = OutputPinLabel;
+	}
+}
+UPCGDynamicMeshData* UPCGGeometryBlueprintElement::CopyOrStealInputData(const FPCGTaggedData& InTaggedData) const
+{
+	if (bIsCacheable || IsCacheableOverride())
+	{
+		// Verification that the user didn't change default settings
+		PCGLog::LogWarningOnGraph(LOCTEXT("SettingsDifferent", "In PCG Geometry Blueprint Element, the default settings were changed (not cacheable and not verifying outputs used multiple times)."
+			"Use the normal BP element if you want this behavior. Will always copy and never steal."), &GetContext());
+		const UPCGDynamicMeshData* InData = Cast<const UPCGDynamicMeshData>(InTaggedData.Data);
+		return InData ? CastChecked<UPCGDynamicMeshData>(InData->DuplicateData(&GetContext())) : nullptr;
+	}
+	
+	return IPCGDynamicMeshBaseElement::CopyOrSteal(InTaggedData, &GetContext());
+}
+
+#undef LOCTEXT_NAMESPACE
diff --git a/Engine/Plugins/PCGInterops/PCGGeometryScriptInterop/Source/PCGGeometryScriptInterop/Public/Elements/PCGDynamicMeshBaseElement.h b/Engine/Plugins/PCGInterops/PCGGeometryScriptInterop/Source/PCGGeometryScriptInterop/Public/Elements/PCGDynamicMeshBaseElement.h
--- a/Engine/Plugins/PCGInterops/PCGGeometryScriptInterop/Source/PCGGeometryScriptInterop/Public/Elements/PCGDynamicMeshBaseElement.h
+++ b/Engine/Plugins/PCGInterops/PCGGeometryScriptInterop/Source/PCGGeometryScriptInterop/Public/Elements/PCGDynamicMeshBaseElement.h
@@ -39,6 +39,7 @@ protected:
 	virtual bool IsCacheable(const UPCGSettings*) const override { return false; }
 	virtual bool ShouldVerifyIfOutputsAreUsedMultipleTimes(const UPCGSettings*) const override { return true; }
 
+public:
 	PCGGEOMETRYSCRIPTINTEROP_API static UPCGDynamicMeshData* CopyOrSteal(const FPCGTaggedData& InTaggedData, FPCGContext* InContext);
 };
 
diff --git a/Engine/Plugins/PCGInterops/PCGGeometryScriptInterop/Source/PCGGeometryScriptInterop/Public/Elements/PCGGeometryBlueprintElement.h b/Engine/Plugins/PCGInterops/PCGGeometryScriptInterop/Source/PCGGeometryScriptInterop/Public/Elements/PCGGeometryBlueprintElement.h
new file mode 100644
--- /dev/null
+++ b/Engine/Plugins/PCGInterops/PCGGeometryScriptInterop/Source/PCGGeometryScriptInterop/Public/Elements/PCGGeometryBlueprintElement.h
@@ -0,0 +1,41 @@
+// Copyright Epic Games, Inc. All Rights Reserved.
+
+#pragma once
+
+#include "Elements/PCGExecuteBlueprint.h"
+
+#include "PCGGeometryBlueprintElement.generated.h"
+
+class UDynamicMesh;
+class UPCGDynamicMeshData;
+
+/**
+* Subclass of PCG Blueprint Element, it comes with pre-configured pins as input and output for Dynamic meshes and force to be non-cacheable.
+* The function CopyOrStealInputData is a helper to either steal (efficient) or copy the input data (less efficient) so work can be done in place on the Dynamic Mesh.
+* More importantly, a user deriving from this class will want to implement ProcessDynamicMesh, the only thing needed to streamline the process and
+* removes all the boilerplate when in a simple input->output case.
+*/
+UCLASS(Abstract, BlueprintType, Blueprintable, hidecategories = (Object))
+class PCGGEOMETRYSCRIPTINTEROP_API UPCGGeometryBlueprintElement : public UPCGBlueprintElement
+{
+	GENERATED_BODY()
+	
+public:
+	UPCGGeometryBlueprintElement();
+
+	/** Virtual implementation to streamline the creation of a Geometry Script node processing. Do not override this function or Execute if you want to use the streamlined version. */
+	virtual void ExecuteWithContext_Implementation(FPCGContext& InContext, const FPCGDataCollection& Input, FPCGDataCollection& Output) override;
+
+	/**
+     * Streamlined version of the Execute function, to only process the dynamic meshes.
+     * For each input that is a dynamic mesh, we will call this function, and it will create as many output data as there are inputs.
+     * @param InDynMesh  - Dynamic mesh to process. Can be used as is and do operation in place.
+     * @param OutTags  - Optional tags to add to the output. By default, it will inherit the tags of the input.
+     */
+    UFUNCTION(BlueprintImplementableEvent, BlueprintCallable, Category = "PCG|Execution", meta = (AutoCreateRefTerm = "InTags"))
+    void ProcessDynamicMesh(UDynamicMesh* InDynMesh, TArray<FString>& OutTags);
+
+	/** Allows to steal the data and work in place if the data is not used elsewhere. If this element is cacheable, it will automatically copy. */
+	UFUNCTION(BlueprintCallable, Category = "PCG|DynamicMesh", meta = (HidePin = "Self"))
+	UPCGDynamicMeshData* CopyOrStealInputData(const FPCGTaggedData& InTaggedData) const;
+};
\ No newline at end of file
