[PCG] Data Crc determinism improvements
- Replace Crc of GetUniqueID() which can change between editor/game runs with class GetPathName() Crc

#rnx

diff --git a/Engine/Plugins/PCG/Source/PCG/Private/Data/PCGCollisionShapeData.cpp b/Engine/Plugins/PCG/Source/PCG/Private/Data/PCGCollisionShapeData.cpp
--- a/Engine/Plugins/PCG/Source/PCG/Private/Data/PCGCollisionShapeData.cpp
+++ b/Engine/Plugins/PCG/Source/PCG/Private/Data/PCGCollisionShapeData.cpp
@@ -47,8 +47,8 @@ void UPCGCollisionShapeData::AddToCrc(FArchiveCrc32& Ar, bool bFullDataCrc) cons
 
 	// Implementation note: no metadata at this point yet.
 
-	uint32 UniqueTypeID = StaticClass()->GetDefaultObject()->GetUniqueID();
-	Ar << UniqueTypeID;
+	FString ClassName = StaticClass()->GetPathName();
+	Ar << ClassName;
 
 	Ar << const_cast<FTransform&>(Transform);
 
diff --git a/Engine/Plugins/PCG/Source/PCG/Private/Data/PCGCollisionWrapperData.cpp b/Engine/Plugins/PCG/Source/PCG/Private/Data/PCGCollisionWrapperData.cpp
--- a/Engine/Plugins/PCG/Source/PCG/Private/Data/PCGCollisionWrapperData.cpp
+++ b/Engine/Plugins/PCG/Source/PCG/Private/Data/PCGCollisionWrapperData.cpp
@@ -275,8 +275,8 @@ void UPCGCollisionWrapperData::AddToCrc(FArchiveCrc32& Ar, bool bFullDataCrc) co
 
 	// Implementation note: no metadata at this point yet
 
-	uint32 UniqueTypeID = StaticClass()->GetDefaultObject()->GetUniqueID();
-	Ar << UniqueTypeID;
+	FString ClassName = StaticClass()->GetPathName();
+	Ar << ClassName;
 
 	CollisionSelector.AddToCrc(Ar);
 
diff --git a/Engine/Plugins/PCG/Source/PCG/Private/Data/PCGDifferenceData.cpp b/Engine/Plugins/PCG/Source/PCG/Private/Data/PCGDifferenceData.cpp
--- a/Engine/Plugins/PCG/Source/PCG/Private/Data/PCGDifferenceData.cpp
+++ b/Engine/Plugins/PCG/Source/PCG/Private/Data/PCGDifferenceData.cpp
@@ -166,8 +166,8 @@ void UPCGDifferenceData::AddToCrc(FArchiveCrc32& Ar, bool bFullDataCrc) const
 
 	// Implementation note: no metadata in composite data at this point.
 
-	uint32 UniqueTypeID = StaticClass()->GetDefaultObject()->GetUniqueID();
-	Ar << UniqueTypeID;
+	FString ClassName = StaticClass()->GetPathName();
+	Ar << ClassName;
 
 	uint32 DiffMetadata = bDiffMetadata ? 1 : 0;
 	Ar << DiffMetadata;
diff --git a/Engine/Plugins/PCG/Source/PCG/Private/Data/PCGIntersectionData.cpp b/Engine/Plugins/PCG/Source/PCG/Private/Data/PCGIntersectionData.cpp
--- a/Engine/Plugins/PCG/Source/PCG/Private/Data/PCGIntersectionData.cpp
+++ b/Engine/Plugins/PCG/Source/PCG/Private/Data/PCGIntersectionData.cpp
@@ -103,8 +103,8 @@ void UPCGIntersectionData::AddToCrc(FArchiveCrc32& Ar, bool bFullDataCrc) const
 
 	// Implementation note: no metadata in composite data at this point.
 
-	uint32 UniqueTypeID = StaticClass()->GetDefaultObject()->GetUniqueID();
-	Ar << UniqueTypeID;
+	FString ClassName = StaticClass()->GetPathName();
+	Ar << ClassName;
 
 	uint32 DensityFunctionValue = static_cast<uint32>(DensityFunction);
 	Ar << DensityFunctionValue;
diff --git a/Engine/Plugins/PCG/Source/PCG/Private/Data/PCGPointData.cpp b/Engine/Plugins/PCG/Source/PCG/Private/Data/PCGPointData.cpp
--- a/Engine/Plugins/PCG/Source/PCG/Private/Data/PCGPointData.cpp
+++ b/Engine/Plugins/PCG/Source/PCG/Private/Data/PCGPointData.cpp
@@ -287,8 +287,8 @@ void UPCGPointData::AddToCrc(FArchiveCrc32& Ar, bool bFullDataCrc) const
 
 	TRACE_CPUPROFILER_EVENT_SCOPE(UPCGPointData::AddToCrc);
 
-	uint32 UniqueTypeID = StaticClass()->GetDefaultObject()->GetUniqueID();
-	Ar << UniqueTypeID;
+	FString ClassName = StaticClass()->GetPathName();
+	Ar << ClassName;
 
 	if (Points.Num() == 0)
 	{
diff --git a/Engine/Plugins/PCG/Source/PCG/Private/Data/PCGProjectionData.cpp b/Engine/Plugins/PCG/Source/PCG/Private/Data/PCGProjectionData.cpp
--- a/Engine/Plugins/PCG/Source/PCG/Private/Data/PCGProjectionData.cpp
+++ b/Engine/Plugins/PCG/Source/PCG/Private/Data/PCGProjectionData.cpp
@@ -100,8 +100,8 @@ void UPCGProjectionData::AddToCrc(FArchiveCrc32& Ar, bool bFullDataCrc) const
 
 	// Implementation note: no metadata in composite data at this point.
 
-	uint32 UniqueTypeID = StaticClass()->GetDefaultObject()->GetUniqueID();
-	Ar << UniqueTypeID;
+	FString ClassName = StaticClass()->GetPathName();
+	Ar << ClassName;
 
 	Ar << ProjectionParams;
 }
diff --git a/Engine/Plugins/PCG/Source/PCG/Private/Data/PCGSplineData.cpp b/Engine/Plugins/PCG/Source/PCG/Private/Data/PCGSplineData.cpp
--- a/Engine/Plugins/PCG/Source/PCG/Private/Data/PCGSplineData.cpp
+++ b/Engine/Plugins/PCG/Source/PCG/Private/Data/PCGSplineData.cpp
@@ -86,8 +86,8 @@ void UPCGSplineData::AddToCrc(FArchiveCrc32& Ar, bool bFullDataCrc) const
 		Metadata->AddToCrc(Ar, bFullDataCrc);
 	}
 
-	uint32 UniqueTypeID = StaticClass()->GetDefaultObject()->GetUniqueID();
-	Ar << UniqueTypeID;
+	FString ClassName = StaticClass()->GetPathName();
+	Ar << ClassName;
 
 	Ar << SplineStruct;
 }
diff --git a/Engine/Plugins/PCG/Source/PCG/Private/Data/PCGSplineInteriorSurfaceData.cpp b/Engine/Plugins/PCG/Source/PCG/Private/Data/PCGSplineInteriorSurfaceData.cpp
--- a/Engine/Plugins/PCG/Source/PCG/Private/Data/PCGSplineInteriorSurfaceData.cpp
+++ b/Engine/Plugins/PCG/Source/PCG/Private/Data/PCGSplineInteriorSurfaceData.cpp
@@ -36,8 +36,8 @@ void UPCGSplineInteriorSurfaceData::AddToCrc(FArchiveCrc32& Ar, bool bFullDataCr
 		Metadata->AddToCrc(Ar, bFullDataCrc);
 	}
 
-	uint32 UniqueTypeID = StaticClass()->GetDefaultObject()->GetUniqueID();
-	Ar << UniqueTypeID;
+	FString ClassName = StaticClass()->GetPathName();
+	Ar << ClassName;
 
 	Ar << SplineStruct;
 }
diff --git a/Engine/Plugins/PCG/Source/PCG/Private/Data/PCGUnionData.cpp b/Engine/Plugins/PCG/Source/PCG/Private/Data/PCGUnionData.cpp
--- a/Engine/Plugins/PCG/Source/PCG/Private/Data/PCGUnionData.cpp
+++ b/Engine/Plugins/PCG/Source/PCG/Private/Data/PCGUnionData.cpp
@@ -120,8 +120,8 @@ void UPCGUnionData::AddToCrc(FArchiveCrc32& Ar, bool bFullDataCrc) const
 
 	// Implementation note: no metadata in composite data at this point.
 
-	uint32 UniqueTypeID = StaticClass()->GetDefaultObject()->GetUniqueID();
-	Ar << UniqueTypeID;
+	FString ClassName = StaticClass()->GetPathName();
+	Ar << ClassName;
 
 	uint32 UnionTypeValue = static_cast<uint32>(UnionType);
 	Ar << UnionTypeValue;
diff --git a/Engine/Plugins/PCG/Source/PCG/Private/Data/PCGVolumeData.cpp b/Engine/Plugins/PCG/Source/PCG/Private/Data/PCGVolumeData.cpp
--- a/Engine/Plugins/PCG/Source/PCG/Private/Data/PCGVolumeData.cpp
+++ b/Engine/Plugins/PCG/Source/PCG/Private/Data/PCGVolumeData.cpp
@@ -88,8 +88,8 @@ void UPCGVolumeData::AddToCrc(FArchiveCrc32& Ar, bool bFullDataCrc) const
 	{
 		// Implementation note: no metadata in this data at this point.
 
-		uint32 UniqueTypeID = StaticClass()->GetDefaultObject()->GetUniqueID();
-		Ar << UniqueTypeID;
+		FString ClassName = StaticClass()->GetPathName();
+		Ar << ClassName;
 
 		// Weird implementation to fix ambiguous call between the FBodyInstance friend function (which isn't exposed) and the one in the FArchiveCrc32 code.
 		FBodyInstance::StaticStruct()->SerializeItem(Ar, const_cast<FBodyInstance*>(VolumeBodyInstance), nullptr);
diff --git a/Engine/Plugins/PCG/Source/PCG/Private/PCGParamData.cpp b/Engine/Plugins/PCG/Source/PCG/Private/PCGParamData.cpp
--- a/Engine/Plugins/PCG/Source/PCG/Private/PCGParamData.cpp
+++ b/Engine/Plugins/PCG/Source/PCG/Private/PCGParamData.cpp
@@ -23,8 +23,8 @@ void UPCGParamData::AddToCrc(FArchiveCrc32& Ar, bool bFullDataCrc) const
 {
 	Super::AddToCrc(Ar, bFullDataCrc);
 
-	uint32 UniqueTypeID = StaticClass()->GetDefaultObject()->GetUniqueID();
-	Ar << UniqueTypeID;
+	FString ClassName = StaticClass()->GetPathName();
+	Ar << ClassName;
 
 	if (Metadata)
 	{
diff --git a/Engine/Plugins/PCGInterops/PCGGeometryScriptInterop/Source/PCGGeometryScriptInterop/Private/Data/PCGDynamicMeshData.cpp b/Engine/Plugins/PCGInterops/PCGGeometryScriptInterop/Source/PCGGeometryScriptInterop/Private/Data/PCGDynamicMeshData.cpp
--- a/Engine/Plugins/PCGInterops/PCGGeometryScriptInterop/Source/PCGGeometryScriptInterop/Private/Data/PCGDynamicMeshData.cpp
+++ b/Engine/Plugins/PCGInterops/PCGGeometryScriptInterop/Source/PCGGeometryScriptInterop/Private/Data/PCGDynamicMeshData.cpp
@@ -74,8 +74,8 @@ void UPCGDynamicMeshData::AddToCrc(FArchiveCrc32& Ar, bool bFullDataCrc) const
 	{
 		// Implementation note: metadata not supported at this point.
 
-		uint32 UniqueTypeID = StaticClass()->GetDefaultObject()->GetUniqueID();
-		Ar << UniqueTypeID;
+		FString ClassName = StaticClass()->GetPathName();
+		Ar << ClassName;
 
 		Ar << DynamicMesh->GetMeshRef();
 	}
