[PCG] Make steepness 1 by default in all newly created samplers

#rnx

diff --git a/Engine/Plugins/PCG/Source/PCG/Private/Elements/PCGCreatePointsGrid.cpp b/Engine/Plugins/PCG/Source/PCG/Private/Elements/PCGCreatePointsGrid.cpp
--- a/Engine/Plugins/PCG/Source/PCG/Private/Elements/PCGCreatePointsGrid.cpp
+++ b/Engine/Plugins/PCG/Source/PCG/Private/Elements/PCGCreatePointsGrid.cpp
@@ -12,6 +12,15 @@
 
 #define LOCTEXT_NAMESPACE "PCGCreatePointsGridElement"
 
+UPCGCreatePointsGridSettings::UPCGCreatePointsGridSettings()
+{
+	if (PCGHelpers::IsNewObjectAndNotDefault(this))
+	{
+		PointSteepness = 1.0f;
+	}
+}
+
+
 void UPCGCreatePointsGridSettings::PostLoad()
 {
 	Super::PostLoad();
diff --git a/Engine/Plugins/PCG/Source/PCG/Private/Elements/PCGCreatePointsSphere.cpp b/Engine/Plugins/PCG/Source/PCG/Private/Elements/PCGCreatePointsSphere.cpp
--- a/Engine/Plugins/PCG/Source/PCG/Private/Elements/PCGCreatePointsSphere.cpp
+++ b/Engine/Plugins/PCG/Source/PCG/Private/Elements/PCGCreatePointsSphere.cpp
@@ -351,6 +351,14 @@ namespace PCGCreatePointsSphere
 	}
 }
 
+UPCGCreatePointsSphereSettings::UPCGCreatePointsSphereSettings()
+{
+	if (PCGHelpers::IsNewObjectAndNotDefault(this))
+	{
+		PointSteepness = 1.0f;
+	}
+}
+
 TArray<FPCGPinProperties> UPCGCreatePointsSphereSettings::OutputPinProperties() const
 {
 	return DefaultPointOutputPinProperties();
diff --git a/Engine/Plugins/PCG/Source/PCG/Private/Elements/PCGSplineSampler.cpp b/Engine/Plugins/PCG/Source/PCG/Private/Elements/PCGSplineSampler.cpp
--- a/Engine/Plugins/PCG/Source/PCG/Private/Elements/PCGSplineSampler.cpp
+++ b/Engine/Plugins/PCG/Source/PCG/Private/Elements/PCGSplineSampler.cpp
@@ -1444,6 +1444,14 @@ namespace PCGSplineSamplerHelpers
 	}
 }
 
+UPCGSplineSamplerSettings::UPCGSplineSamplerSettings()
+{
+	if (PCGHelpers::IsNewObjectAndNotDefault(this))
+	{
+		SamplerParams.PointSteepness = 1.0f;
+	}
+}
+
 #if WITH_EDITOR
 FText UPCGSplineSamplerSettings::GetNodeTooltipText() const
 {
diff --git a/Engine/Plugins/PCG/Source/PCG/Private/Elements/PCGSurfaceSampler.cpp b/Engine/Plugins/PCG/Source/PCG/Private/Elements/PCGSurfaceSampler.cpp
--- a/Engine/Plugins/PCG/Source/PCG/Private/Elements/PCGSurfaceSampler.cpp
+++ b/Engine/Plugins/PCG/Source/PCG/Private/Elements/PCGSurfaceSampler.cpp
@@ -266,6 +266,14 @@ namespace PCGSurfaceSampler
 #endif
 }
 
+UPCGSurfaceSamplerSettings::UPCGSurfaceSamplerSettings()
+{
+	if (PCGHelpers::IsNewObjectAndNotDefault(this))
+	{
+		PointSteepness = 1.0f;
+	}
+}
+
 #if WITH_EDITOR
 FText UPCGSurfaceSamplerSettings::GetNodeTooltipText() const
 {
diff --git a/Engine/Plugins/PCG/Source/PCG/Private/Elements/PCGVolumeSampler.cpp b/Engine/Plugins/PCG/Source/PCG/Private/Elements/PCGVolumeSampler.cpp
--- a/Engine/Plugins/PCG/Source/PCG/Private/Elements/PCGVolumeSampler.cpp
+++ b/Engine/Plugins/PCG/Source/PCG/Private/Elements/PCGVolumeSampler.cpp
@@ -134,6 +134,13 @@ namespace PCGVolumeSampler
 		return FPCGAsync::AsyncProcessing<FPCGPoint>(AsyncState, NumIterations, Points, AsyncProcessingFunc, /*bEnableTimeSlicing=*/Context && bTimeSlicingIsEnabled);
 	}
 }
+UPCGVolumeSamplerSettings::UPCGVolumeSamplerSettings()
+{
+	if (PCGHelpers::IsNewObjectAndNotDefault(this))
+	{
+		PointSteepness = 1.0f;
+	}
+}
 
 #if WITH_EDITOR
 FText UPCGVolumeSamplerSettings::GetNodeTooltipText() const
diff --git a/Engine/Plugins/PCG/Source/PCG/Public/Elements/PCGCreatePointsGrid.h b/Engine/Plugins/PCG/Source/PCG/Public/Elements/PCGCreatePointsGrid.h
--- a/Engine/Plugins/PCG/Source/PCG/Public/Elements/PCGCreatePointsGrid.h
+++ b/Engine/Plugins/PCG/Source/PCG/Public/Elements/PCGCreatePointsGrid.h
@@ -30,6 +30,8 @@ class UPCGCreatePointsGridSettings : public UPCGSettings
 	GENERATED_BODY()
 
 public:
+	UPCGCreatePointsGridSettings();
+
 	//~Begin UObject interface
 	virtual void PostLoad() override;
 	//~End UObject interface
diff --git a/Engine/Plugins/PCG/Source/PCG/Public/Elements/PCGCreatePointsSphere.h b/Engine/Plugins/PCG/Source/PCG/Public/Elements/PCGCreatePointsSphere.h
--- a/Engine/Plugins/PCG/Source/PCG/Public/Elements/PCGCreatePointsSphere.h
+++ b/Engine/Plugins/PCG/Source/PCG/Public/Elements/PCGCreatePointsSphere.h
@@ -41,6 +41,8 @@ class UPCGCreatePointsSphereSettings : public UPCGSettings
 	GENERATED_BODY()
 
 public:
+	UPCGCreatePointsSphereSettings();
+
 	//~Begin UPCGSettings interface
 #if WITH_EDITOR
 	virtual FName GetDefaultNodeName() const override { return FName(TEXT("CreatePointsSphere")); }
diff --git a/Engine/Plugins/PCG/Source/PCG/Public/Elements/PCGSplineSampler.h b/Engine/Plugins/PCG/Source/PCG/Public/Elements/PCGSplineSampler.h
--- a/Engine/Plugins/PCG/Source/PCG/Public/Elements/PCGSplineSampler.h
+++ b/Engine/Plugins/PCG/Source/PCG/Public/Elements/PCGSplineSampler.h
@@ -254,6 +254,8 @@ class UPCGSplineSamplerSettings : public UPCGSettings
 	GENERATED_BODY()
 
 public:
+	UPCGSplineSamplerSettings();
+
 	// ~Begin UPCGSettings interface
 #if WITH_EDITOR
 	virtual FName GetDefaultNodeName() const override { return FName(TEXT("SplineSampler")); }
diff --git a/Engine/Plugins/PCG/Source/PCG/Public/Elements/PCGSurfaceSampler.h b/Engine/Plugins/PCG/Source/PCG/Public/Elements/PCGSurfaceSampler.h
--- a/Engine/Plugins/PCG/Source/PCG/Public/Elements/PCGSurfaceSampler.h
+++ b/Engine/Plugins/PCG/Source/PCG/Public/Elements/PCGSurfaceSampler.h
@@ -94,6 +94,8 @@ class UPCGSurfaceSamplerSettings : public UPCGSettings
 	GENERATED_BODY()
 
 public:
+	UPCGSurfaceSamplerSettings();
+
 	//~Begin UObject interface
 	virtual void PostLoad() override;
 	//~End UObject interface
diff --git a/Engine/Plugins/PCG/Source/PCG/Public/Elements/PCGVolumeSampler.h b/Engine/Plugins/PCG/Source/PCG/Public/Elements/PCGVolumeSampler.h
--- a/Engine/Plugins/PCG/Source/PCG/Public/Elements/PCGVolumeSampler.h
+++ b/Engine/Plugins/PCG/Source/PCG/Public/Elements/PCGVolumeSampler.h
@@ -57,6 +57,8 @@ class UPCGVolumeSamplerSettings : public UPCGSettings
 	GENERATED_BODY()
 
 public:
+	UPCGVolumeSamplerSettings();
+
 	UPROPERTY(BlueprintReadWrite, EditAnywhere, Category = "Data", meta = (PCG_Overridable))
 	FVector VoxelSize = PCGVolumeSampler::DefaultVoxelSize;
 
diff --git a/Engine/Plugins/PCGInterops/PCGGeometryScriptInterop/Source/PCGGeometryScriptInterop/Private/Elements/PCGMeshSampler.cpp b/Engine/Plugins/PCGInterops/PCGGeometryScriptInterop/Source/PCGGeometryScriptInterop/Private/Elements/PCGMeshSampler.cpp
--- a/Engine/Plugins/PCGInterops/PCGGeometryScriptInterop/Source/PCGGeometryScriptInterop/Private/Elements/PCGMeshSampler.cpp
+++ b/Engine/Plugins/PCGInterops/PCGGeometryScriptInterop/Source/PCGGeometryScriptInterop/Private/Elements/PCGMeshSampler.cpp
@@ -292,6 +292,14 @@ void FPCGMeshSamplerContext::SetPointColorAndDensity(SetPointDensityFunc SetPoin
 	}
 }
 
+UPCGMeshSamplerSettings::UPCGMeshSamplerSettings()
+{
+	if (PCGHelpers::IsNewObjectAndNotDefault(this))
+	{
+		PointSteepness = 1.0f;
+	}
+}
+
 void UPCGMeshSamplerSettings::PostLoad()
 {
 	Super::PostLoad();
diff --git a/Engine/Plugins/PCGInterops/PCGGeometryScriptInterop/Source/PCGGeometryScriptInterop/Public/Elements/PCGMeshSampler.h b/Engine/Plugins/PCGInterops/PCGGeometryScriptInterop/Source/PCGGeometryScriptInterop/Public/Elements/PCGMeshSampler.h
--- a/Engine/Plugins/PCGInterops/PCGGeometryScriptInterop/Source/PCGGeometryScriptInterop/Public/Elements/PCGMeshSampler.h
+++ b/Engine/Plugins/PCGInterops/PCGGeometryScriptInterop/Source/PCGGeometryScriptInterop/Public/Elements/PCGMeshSampler.h
@@ -57,6 +57,8 @@ class PCGGEOMETRYSCRIPTINTEROP_API UPCGMeshSamplerSettings : public UPCGSettings
 	GENERATED_BODY()
 
 public:
+	UPCGMeshSamplerSettings();
+
 	//~Begin UObject interface
 	virtual void PostLoad() override;
 	//~End UObject interface
