[PCG] Fix crash on PCG GeometryScript Interop module shutdown

#rnx

diff --git a/Engine/Plugins/PCGInterops/PCGGeometryScriptInterop/Source/PCGGeometryScriptInterop/Private/PCGGeometryScriptInteropModule.cpp b/Engine/Plugins/PCGInterops/PCGGeometryScriptInterop/Source/PCGGeometryScriptInterop/Private/PCGGeometryScriptInteropModule.cpp
--- a/Engine/Plugins/PCGInterops/PCGGeometryScriptInterop/Source/PCGGeometryScriptInterop/Private/PCGGeometryScriptInteropModule.cpp
+++ b/Engine/Plugins/PCGInterops/PCGGeometryScriptInterop/Source/PCGGeometryScriptInterop/Private/PCGGeometryScriptInteropModule.cpp
@@ -6,6 +6,8 @@
 #include "Data/PCGDynamicMeshData.h"
 #include "Elements/PCGGetDynamicMeshData.h"
 
+#include "Misc/CoreDelegates.h"
+
 #if WITH_EDITOR
 #include "Editor/PCGDynamicMeshDataVisualization.h"
 #endif // WITH_EDITOR
@@ -13,12 +15,28 @@
 class FPCGGeometryScriptInteropModule final : public IModuleInterface
 {
 public:
+	void PreExit()
+	{
+		// No need to unregister if the PCG module is already dead.
+		if (FPCGModule::IsPCGModuleLoaded())
+		{
+#if WITH_EDITOR
+			FPCGDataVisualizationRegistry& DataVisRegistry = FPCGModule::GetMutablePCGDataVisualizationRegistry();
+			DataVisRegistry.UnregisterPCGDataVisualization(UPCGDynamicMeshData::StaticClass());
+#endif // WITH_EDITOR
+
+			FPCGGetDataFunctionRegistry& PCGDataFunctionRegistry = FPCGModule::MutableGetDataFunctionRegistry();
+			PCGDataFunctionRegistry.UnregisterDataFromActorFunction(GetActorDataFunctionHandle);
+			PCGDataFunctionRegistry.UnregisterDataFromComponentFunction(GetComponentDataFunctionHandle);
+		}
+	}
+
 	//~ IModuleInterface implementation
 	virtual bool SupportsDynamicReloading() override
 	{
 		return true;
 	}
-	
+
 	virtual void StartupModule() override
 	{
 		FModuleManager::Get().LoadModuleChecked(TEXT("PCG"));
@@ -31,22 +49,14 @@ public:
 		FPCGGetDataFunctionRegistry& PCGDataFunctionRegistry = FPCGModule::MutableGetDataFunctionRegistry();
 		GetActorDataFunctionHandle = PCGDataFunctionRegistry.RegisterDataFromActorFunction(&PCGGetDynamicMeshData::GetDynamicMeshDataFromActor);
 		GetComponentDataFunctionHandle = PCGDataFunctionRegistry.RegisterDataFromComponentFunction(&PCGGetDynamicMeshData::GetDynamicMeshDataFromComponent);
+
+		// Register onto the PreExit, because we need the class to be still valid to remove them from the mapping
+		FCoreDelegates::OnPreExit.AddRaw(this, &FPCGGeometryScriptInteropModule::PreExit);
 	}
 
 	virtual void ShutdownModule() override
 	{
-		// No need to unregister if the PCG module is already dead.
-		if (FPCGModule::IsPCGModuleLoaded())
-		{
-			FPCGGetDataFunctionRegistry& PCGDataFunctionRegistry = FPCGModule::MutableGetDataFunctionRegistry();
-			PCGDataFunctionRegistry.UnregisterDataFromActorFunction(GetActorDataFunctionHandle);
-			PCGDataFunctionRegistry.UnregisterDataFromComponentFunction(GetComponentDataFunctionHandle);
-		
-#if WITH_EDITOR
-			FPCGDataVisualizationRegistry& DataVisRegistry = FPCGModule::GetMutablePCGDataVisualizationRegistry();
-			DataVisRegistry.UnregisterPCGDataVisualization(UPCGDynamicMeshData::StaticClass());
-#endif // WITH_EDITOR
-		}
+		FCoreDelegates::OnPreExit.RemoveAll(this);
 	}
 	//~ End IModuleInterface implementation
 	
