[PCG] Tracking improvement

* When a component changes, we take its owner actor and test the actor against the tracking keys
* But with the new Get Actor Data from components, the tracking need also to react to component changes.
* Same for Mesh Sampler and dynamic mesh components.
* The Mesh sampler was also not setup to track anything
* Also integration from main for: fix issue when moving actors that have no bounds in partitioned
* Also add a condition to not discard previous data if there is a no op in Object Changed because the actor is delayed

#rnx

diff --git a/Engine/Plugins/PCG/Source/PCG/Private/PCGActorAndComponentMapping.cpp b/Engine/Plugins/PCG/Source/PCG/Private/PCGActorAndComponentMapping.cpp
--- a/Engine/Plugins/PCG/Source/PCG/Private/PCGActorAndComponentMapping.cpp
+++ b/Engine/Plugins/PCG/Source/PCG/Private/PCGActorAndComponentMapping.cpp
@@ -1691,7 +1691,7 @@ void FPCGActorAndComponentMapping::OnObjectModified(UObject* InObject)
 void FPCGActorAndComponentMapping::OnObjectSaved(UObject* InObject, FObjectPreSaveContext InObjectSaveContext)
 {
 	// Nothing to do if we track nothing
-	if (CulledTrackedKeysToComponentsMap.IsEmpty() && AlwaysTrackedKeysToComponentsMap.IsEmpty())
+	if (!InObject || (CulledTrackedKeysToComponentsMap.IsEmpty() && AlwaysTrackedKeysToComponentsMap.IsEmpty()))
 	{
 		return;
 	}
@@ -1699,10 +1699,10 @@ void FPCGActorAndComponentMapping::OnObjectSaved(UObject* InObject, FObjectPreSa
 	// Only trigger a refresh on save for limited data classes.
 	// At this point in time, We only track data tables and PCG data assets because in most cases we probably will catch other changes with OnObjectPropertyChanged.
 	// This is especially important to make sure we don't trigger refresh multiple times (less a problem for PCG assets, but a big problem for data tables).
-	if (!InObjectSaveContext.IsProceduralSave() && (Cast<UDataTable>(InObject) || Cast<UPCGDataAsset>(InObject)))
+	if (!InObjectSaveContext.IsProceduralSave() && (InObject->IsA<UDataTable>() || InObject->IsA<UPCGDataAsset>() || InObject->IsA<UStaticMesh>()))
 	{
 		FPropertyChangedEvent Event{ nullptr };
-		OnObjectPropertyChanged(InObject, Event);
+		OnObjectChanged(InObject, /*InPreviousData=*/nullptr, /*InOriginatingChangeObject=*/ InObject);
 	}
 }
 
@@ -1823,11 +1823,12 @@ void FPCGActorAndComponentMapping::OnObjectChanged(UObject* InObject, const FAct
 	}
 
 	FSoftObjectPath ObjectPath(InObject);
-	auto Gather = [InObject, &ObjectPath, &RemovedTags, &MatchedKeys](TMap<FPCGSelectionKey, TSet<UPCGComponent*>> InMap, TSet<UPCGComponent*>& OutSet)
+	FSoftObjectPath InOriginatingChangeObjectPath(InOriginatingChangeObject);
+	auto Gather = [InObject, &ObjectPath, &RemovedTags, &MatchedKeys, InOriginatingChangeObject, &InOriginatingChangeObjectPath](TMap<FPCGSelectionKey, TSet<UPCGComponent*>> InMap, TSet<UPCGComponent*>& OutSet)
 	{
 		for (auto& It : InMap)
 		{
-			if (It.Key.IsMatching(InObject, ObjectPath, RemovedTags, It.Value, &OutSet))
+			if (It.Key.IsMatching(InObject, ObjectPath, RemovedTags, It.Value, &OutSet) || It.Key.IsMatching(InOriginatingChangeObject, InOriginatingChangeObjectPath, RemovedTags, It.Value, &OutSet))
 			{
 				MatchedKeys.Add(It.Key);
 			}
diff --git a/Engine/Plugins/PCGInterops/PCGGeometryScriptInterop/Source/PCGGeometryScriptInterop/Private/Elements/PCGMeshSampler.cpp b/Engine/Plugins/PCGInterops/PCGGeometryScriptInterop/Source/PCGGeometryScriptInterop/Private/Elements/PCGMeshSampler.cpp
--- a/Engine/Plugins/PCGInterops/PCGGeometryScriptInterop/Source/PCGGeometryScriptInterop/Private/Elements/PCGMeshSampler.cpp
+++ b/Engine/Plugins/PCGInterops/PCGGeometryScriptInterop/Source/PCGGeometryScriptInterop/Private/Elements/PCGMeshSampler.cpp
@@ -10,6 +10,7 @@
 #include "Data/PCGPointData.h"
 #include "Helpers/PCGAsync.h"
 #include "Helpers/PCGBlueprintHelpers.h"
+#include "Helpers/PCGDynamicTrackingHelpers.h"
 #include "Helpers/PCGHelpers.h"
 #include "Metadata/PCGMetadata.h"
 #include "Metadata/PCGMetadataAttributeTpl.h"
@@ -349,6 +350,18 @@ FText UPCGMeshSamplerSettings::GetNodeTooltipText() const
 {
 	return LOCTEXT("MeshSamplerNodeTooltip", "Sample points on a static mesh.");
 }
+
+void UPCGMeshSamplerSettings::GetStaticTrackedKeys(FPCGSelectionKeyToSettingsMap& OutKeysToSettings, TArray<TObjectPtr<const UPCGGraph>>& OutVisitedGraphs) const
+{
+	if (StaticMesh.IsNull())
+	{
+		return;
+	}
+
+	FPCGSelectionKey Key = FPCGSelectionKey::CreateFromPath(StaticMesh.ToSoftObjectPath());
+
+	OutKeysToSettings.FindOrAdd(Key).Emplace(this, /*bCulling=*/false);
+}
 #endif // WITH_EDITOR
 
 FPCGMeshSamplerContext::~FPCGMeshSamplerContext()
@@ -411,6 +424,11 @@ bool FPCGMeshSamplerElement::PrepareDataInternal(FPCGContext* InContext) const
 
 	UGeometryScriptDebug* Debug = FPCGContext::NewObject_AnyThread<UGeometryScriptDebug>(Context);
 
+#if WITH_EDITOR
+	FPCGDynamicTrackingHelper DynamicTracking;
+	DynamicTracking.EnableAndInitialize(Context);
+#endif // WITH_EDITOR
+
 	for (const auto& [Path, DummyIndex, DummyIndex2] : Context->PathsToObjectsAndDataIndex)
 	{
 		UObject* Object = Path.ResolveObject();
@@ -432,6 +450,10 @@ bool FPCGMeshSamplerElement::PrepareDataInternal(FPCGContext* InContext) const
 
 		if (SceneComponent)
 		{
+#if WITH_EDITOR
+			DynamicTracking.AddToTracking(FPCGSelectionKey::CreateFromPath(FSoftObjectPath(SceneComponent)), /*bIsCulled=*/false);
+#endif // WITH_EDITOR
+
 			// Adaptation of UGeometryScriptLibrary_SceneUtilityFunctions::CopyMeshFromComponent, since we don't have access to the Material list with this one.
 			Context->DynamicMeshes.Add(FPCGContext::NewObject_AnyThread<UDynamicMesh>(Context));
 			FTransform Transform;
@@ -461,6 +483,11 @@ bool FPCGMeshSamplerElement::PrepareDataInternal(FPCGContext* InContext) const
 		}
 		else if (UStaticMesh* StaticMesh = Cast<UStaticMesh>(Object))
 		{
+
+#if WITH_EDITOR
+			DynamicTracking.AddToTracking(FPCGSelectionKey::CreateFromPath(FSoftObjectPath(StaticMesh)), /*bIsCulled=*/false);
+#endif // WITH_EDITOR
+
 			const FGeometryScriptMeshReadLOD MeshReadLOD{ Settings->RequestedLODType, Settings->RequestedLODIndex };
 
 			Context->DynamicMeshes.Add(FPCGContext::NewObject_AnyThread<UDynamicMesh>(Context));
@@ -488,6 +515,10 @@ bool FPCGMeshSamplerElement::PrepareDataInternal(FPCGContext* InContext) const
 		}
 	}
 
+#if WITH_EDITOR
+	DynamicTracking.Finalize(Context);
+#endif // WITH_EDITOR
+
 	// Manually adding incoming Dynamic meshes data input, copy is necessary if we voxelize, otherwise we'll const cast as it won't be modified (but GeometryScript API is not const friendly).
 	for (const FPCGTaggedData& Input : InContext->InputData.GetInputsByPin(PCGPinConstants::DefaultInputLabel))
 	{
diff --git a/Engine/Plugins/PCGInterops/PCGGeometryScriptInterop/Source/PCGGeometryScriptInterop/Public/Elements/PCGMeshSampler.h b/Engine/Plugins/PCGInterops/PCGGeometryScriptInterop/Source/PCGGeometryScriptInterop/Public/Elements/PCGMeshSampler.h
--- a/Engine/Plugins/PCGInterops/PCGGeometryScriptInterop/Source/PCGGeometryScriptInterop/Public/Elements/PCGMeshSampler.h
+++ b/Engine/Plugins/PCGInterops/PCGGeometryScriptInterop/Source/PCGGeometryScriptInterop/Public/Elements/PCGMeshSampler.h
@@ -69,6 +69,8 @@ public:
 	virtual FText GetDefaultNodeTitle() const override { return NSLOCTEXT("PCGMeshSamplerSettings", "NodeTitle", "Mesh Sampler"); }
 	virtual FText GetNodeTooltipText() const override;
 	virtual EPCGSettingsType GetType() const override { return EPCGSettingsType::Sampler; }
+	virtual bool CanDynamicallyTrackKeys() const override { return bExtractMeshFromInput; }
+	virtual void GetStaticTrackedKeys(FPCGSelectionKeyToSettingsMap& OutKeysToSettings, TArray<TObjectPtr<const UPCGGraph>>& OutVisitedGraphs) const override;
 #endif
 	virtual bool UseSeed() const override { return true; }
 	virtual TArray<FPCGPinProperties> InputPinProperties() const override;
