[PCG] Fix Append Meshes from Points when static meshes comes from attributes

#rnx

diff --git a/Engine/Plugins/PCGInterops/PCGGeometryScriptInterop/Source/PCGGeometryScriptInterop/Private/Elements/PCGAppendMeshesFromPoints.cpp b/Engine/Plugins/PCGInterops/PCGGeometryScriptInterop/Source/PCGGeometryScriptInterop/Private/Elements/PCGAppendMeshesFromPoints.cpp
--- a/Engine/Plugins/PCGInterops/PCGGeometryScriptInterop/Source/PCGGeometryScriptInterop/Private/Elements/PCGAppendMeshesFromPoints.cpp
+++ b/Engine/Plugins/PCGInterops/PCGGeometryScriptInterop/Source/PCGGeometryScriptInterop/Private/Elements/PCGAppendMeshesFromPoints.cpp
@@ -207,14 +207,14 @@ bool FPCGAppendMeshesFromPointsElement::ExecuteInternal(FPCGContext* InContext)
 
 	TMap<FSoftObjectPath, UE::Geometry::FDynamicMesh3> StaticMeshToDynMesh;
 
-	auto ConvertStaticMesh = [Settings, &StaticMeshToDynMesh, InContext, &OutDynMeshData](const FSoftObjectPath& InPath) -> bool
+	auto ConvertStaticMesh = [Settings, &StaticMeshToDynMesh, InContext, &OutDynMeshData](const TSoftObjectPtr<UStaticMesh>& InMesh) -> bool
 	{
 		UE::Conversion::FStaticMeshConversionOptions ConversionOptions{};
 		FText ErrorMessage;
 		UE::Conversion::EMeshLODType LODType = static_cast<UE::Conversion::EMeshLODType>(Settings->RequestedLODType);
 		
-		UE::Geometry::FDynamicMesh3& NewMesh = StaticMeshToDynMesh.Add(InPath);
-		UStaticMesh* StaticMesh = Settings->StaticMesh.Get();
+		UE::Geometry::FDynamicMesh3& NewMesh = StaticMeshToDynMesh.Add(InMesh.ToSoftObjectPath());
+		UStaticMesh* StaticMesh = InMesh.LoadSynchronous();
 	
 		if (!UE::Conversion::StaticMeshToDynamicMesh(StaticMesh, NewMesh, ErrorMessage, ConversionOptions, LODType, Settings->RequestedLODIndex))
 		{
@@ -242,7 +242,7 @@ bool FPCGAppendMeshesFromPointsElement::ExecuteInternal(FPCGContext* InContext)
 	{
 	case EPCGAppendMeshesFromPointsMode::SingleStaticMesh:
 	{
-		if (!ConvertStaticMesh(Settings->StaticMesh.ToSoftObjectPath()))
+		if (!ConvertStaticMesh(Settings->StaticMesh))
 		{
 			return true;
 		}
@@ -255,7 +255,7 @@ bool FPCGAppendMeshesFromPointsElement::ExecuteInternal(FPCGContext* InContext)
 	
 		for (TPair<FSoftObjectPath, TArray<int32>>& It : Context->MeshToPointIndicesMapping)
 		{
-			if (!ConvertStaticMesh(It.Key))
+			if (!ConvertStaticMesh(TSoftObjectPtr<UStaticMesh>(It.Key)))
 			{
 				return true;
 			}
