[PCG] Add support for dynamic mesh data to the data viewport

diff --git a/Engine/Plugins/PCGInterops/PCGGeometryScriptInterop/Source/PCGGeometryScriptInterop/PCGGeometryScriptInterop.Build.cs b/Engine/Plugins/PCGInterops/PCGGeometryScriptInterop/Source/PCGGeometryScriptInterop/PCGGeometryScriptInterop.Build.cs
--- a/Engine/Plugins/PCGInterops/PCGGeometryScriptInterop/Source/PCGGeometryScriptInterop/PCGGeometryScriptInterop.Build.cs
+++ b/Engine/Plugins/PCGInterops/PCGGeometryScriptInterop/Source/PCGGeometryScriptInterop/PCGGeometryScriptInterop.Build.cs
@@ -35,6 +35,7 @@ namespace UnrealBuildTool.Rules
 				PrivateDependencyModuleNames.AddRange(
 					new string[]
 					{
+						"AdvancedPreviewScene",
 						"UnrealEd",
 						"PCGEditor"
 					}
diff --git a/Engine/Plugins/PCGInterops/PCGGeometryScriptInterop/Source/PCGGeometryScriptInterop/Private/Editor/PCGDynamicMeshDataVisualization.cpp b/Engine/Plugins/PCGInterops/PCGGeometryScriptInterop/Source/PCGGeometryScriptInterop/Private/Editor/PCGDynamicMeshDataVisualization.cpp
--- a/Engine/Plugins/PCGInterops/PCGGeometryScriptInterop/Source/PCGGeometryScriptInterop/Private/Editor/PCGDynamicMeshDataVisualization.cpp
+++ b/Engine/Plugins/PCGInterops/PCGGeometryScriptInterop/Source/PCGGeometryScriptInterop/Private/Editor/PCGDynamicMeshDataVisualization.cpp
@@ -10,6 +10,7 @@
 #include "Helpers/PCGHelpers.h"
 #include "Resources/PCGDynamicMeshManagedComponent.h"
 
+#include "AdvancedPreviewScene.h"
 #include "UDynamicMesh.h"
 #include "Components/DynamicMeshComponent.h"
 
@@ -35,6 +36,42 @@ void FPCGDynamicMeshDataVisualization::ExecuteDebugDisplay(FPCGContext* Context,
 	}
 }
 
+FPCGSetupSceneFunc FPCGDynamicMeshDataVisualization::GetViewportSetupFunc(const UPCGData* Data) const
+{
+	return [this, WeakData=TWeakObjectPtr<const UPCGDynamicMeshData>(Cast<UPCGDynamicMeshData>(Data))](FAdvancedPreviewScene* InScene, const TArray<UObject*>& InResources, FVector& OutViewLocation, float& OutOrbitDistance, TArray<TObjectPtr<UObject>>& OutManagedResources)
+	{
+		check(InScene);
+
+		if (!WeakData.IsValid())
+		{
+			UE_LOG(LogPCG, Error, TEXT("Failed to setup data viewport, the data was lost or invalid."));
+			return;
+		}
+
+		TObjectPtr<UDynamicMeshComponent> DynamicMeshComponent = NewObject<UDynamicMeshComponent>(GetTransientPackage(), NAME_None, RF_Transient);
+		WeakData->InitializeDynamicMeshComponentFromData(DynamicMeshComponent);
+
+		OutManagedResources.Add(DynamicMeshComponent);
+
+		if (GEditor->PreviewPlatform.GetEffectivePreviewFeatureLevel() <= ERHIFeatureLevel::ES3_1)
+		{
+			DynamicMeshComponent->SetMobility(EComponentMobility::Static);
+		}
+
+		InScene->AddComponent(DynamicMeshComponent, FTransform::Identity);
+
+		const FBoxSphereBounds Bounds = DynamicMeshComponent->CalcLocalBounds();
+		InScene->SetFloorOffset(-Bounds.Origin.Z + Bounds.BoxExtent.Z);
+
+		OutViewLocation = Bounds.Origin;
+
+		if (Bounds.SphereRadius > 0.0f)
+		{
+			OutOrbitDistance = Bounds.SphereRadius * 2.0f;
+		}
+	};
+}
+
 #undef LOCTEXT_NAMESPACE
 
 #endif // WITH_EDITOR
\ No newline at end of file
diff --git a/Engine/Plugins/PCGInterops/PCGGeometryScriptInterop/Source/PCGGeometryScriptInterop/Private/Editor/PCGDynamicMeshDataVisualization.h b/Engine/Plugins/PCGInterops/PCGGeometryScriptInterop/Source/PCGGeometryScriptInterop/Private/Editor/PCGDynamicMeshDataVisualization.h
--- a/Engine/Plugins/PCGInterops/PCGGeometryScriptInterop/Source/PCGGeometryScriptInterop/Private/Editor/PCGDynamicMeshDataVisualization.h
+++ b/Engine/Plugins/PCGInterops/PCGGeometryScriptInterop/Source/PCGGeometryScriptInterop/Private/Editor/PCGDynamicMeshDataVisualization.h
@@ -15,6 +15,7 @@ class FPCGDynamicMeshDataVisualization : public IPCGSpatialDataVisualization
 public:
 	// ~Begin IPCGDataVisualization interface
 	virtual void ExecuteDebugDisplay(FPCGContext* Context, const UPCGSettingsInterface* SettingsInterface, const UPCGData* Data, AActor* TargetActor) const override;
+	virtual FPCGSetupSceneFunc GetViewportSetupFunc(const UPCGData* Data) const override;
 	// ~End IPCGDataVisualization interface
 };
 
