[PCG] Fix issue where dynamic meshes were not producing the right result when inspecting and in viewport

* Don't steal Dynamic meshes if inspection is enabled

#rnx

diff --git a/Engine/Plugins/PCGInterops/PCGGeometryScriptInterop/Source/PCGGeometryScriptInterop/Private/Elements/PCGDynamicMeshBaseElement.cpp b/Engine/Plugins/PCGInterops/PCGGeometryScriptInterop/Source/PCGGeometryScriptInterop/Private/Elements/PCGDynamicMeshBaseElement.cpp
--- a/Engine/Plugins/PCGInterops/PCGGeometryScriptInterop/Source/PCGGeometryScriptInterop/Private/Elements/PCGDynamicMeshBaseElement.cpp
+++ b/Engine/Plugins/PCGInterops/PCGGeometryScriptInterop/Source/PCGGeometryScriptInterop/Private/Elements/PCGDynamicMeshBaseElement.cpp
@@ -4,6 +4,7 @@
 
 #include "PCGContext.h"
 #include "PCGData.h"
+#include "PCGGraphExecutionInspection.h"
 #include "Data/PCGDynamicMeshData.h"
 
 #define LOCTEXT_NAMESPACE "PCGDynamicMeshBaseElement"
@@ -46,8 +47,18 @@ UPCGDynamicMeshData* IPCGDynamicMeshBaseElement::CopyOrSteal(const FPCGTaggedDat
 	{
 		return nullptr;
 	}
+
+	bool bCanSteal = !InTaggedData.bIsUsedMultipleTimes && PCGDynamicMeshBaseElement::CVarPCGDynamicMeshAllowSteal.GetValueOnAnyThread();
+
+#if WITH_EDITOR
+	// We can't steal when we are inspecting, otherwise it breaks the inspection cache.
+	if (InContext->ExecutionSource.IsValid() && InContext->ExecutionSource->GetExecutionState().GetInspection().IsInspecting())
+	{
+		bCanSteal = false; 
+	}
+#endif // WITH_EDITOR
 	
-	if (InTaggedData.bIsUsedMultipleTimes || !PCGDynamicMeshBaseElement::CVarPCGDynamicMeshAllowSteal.GetValueOnAnyThread())
+	if (!bCanSteal)
 	{
 		return CastChecked<UPCGDynamicMeshData>(InData->DuplicateData(InContext));
 	}
