[PCG] Add static mesh tracking to nodes that depends on it

* Also don't track keys statically if we override

#rnx

diff --git a/Engine/Plugins/PCG/Source/PCG/Private/Elements/IO/PCGDataTableElement.cpp b/Engine/Plugins/PCG/Source/PCG/Private/Elements/IO/PCGDataTableElement.cpp
--- a/Engine/Plugins/PCG/Source/PCG/Private/Elements/IO/PCGDataTableElement.cpp
+++ b/Engine/Plugins/PCG/Source/PCG/Private/Elements/IO/PCGDataTableElement.cpp
@@ -7,7 +7,10 @@
 #include "PCGParamData.h"
 #include "Data/PCGBasePointData.h"
 #include "Elements/IO/PCGExternalDataContext.h"
+
+#if WITH_EDITOR
 #include "Helpers/PCGDynamicTrackingHelpers.h"
+#endif // WITH_EDITOR
 
 #include "Engine/DataTable.h"
 
@@ -16,7 +19,7 @@
 #if WITH_EDITOR
 void UPCGLoadDataTableSettings::GetStaticTrackedKeys(FPCGSelectionKeyToSettingsMap& OutKeysToSettings, TArray<TObjectPtr<const UPCGGraph>>& OutVisitedGraphs) const
 {
-	if (DataTable.IsNull())
+	if (DataTable.IsNull() || IsPropertyOverriddenByPin(GET_MEMBER_NAME_CHECKED(UPCGLoadDataTableSettings, DataTable)))
 	{
 		return;
 	}
diff --git a/Engine/Plugins/PCG/Source/PCG/Private/Elements/IO/PCGLoadAssetElement.cpp b/Engine/Plugins/PCG/Source/PCG/Private/Elements/IO/PCGLoadAssetElement.cpp
--- a/Engine/Plugins/PCG/Source/PCG/Private/Elements/IO/PCGLoadAssetElement.cpp
+++ b/Engine/Plugins/PCG/Source/PCG/Private/Elements/IO/PCGLoadAssetElement.cpp
@@ -5,11 +5,14 @@
 #include "PCGModule.h"
 #include "PCGData.h"
 #include "PCGParamData.h"
-#include "Helpers/PCGDynamicTrackingHelpers.h"
 #include "Helpers/PCGHelpers.h"
 #include "Metadata/Accessors/PCGAttributeAccessor.h"
 #include "Metadata/Accessors/PCGAttributeAccessorHelpers.h"
 
+#if WITH_EDITOR
+#include "Helpers/PCGDynamicTrackingHelpers.h"
+#endif // WITH_EDITOR
+
 #include "Algo/Transform.h"
 #include "AssetRegistry/AssetData.h"
 
@@ -29,7 +32,7 @@ UPCGLoadDataAssetSettings::UPCGLoadDataAssetSettings()
 #if WITH_EDITOR
 void UPCGLoadDataAssetSettings::GetStaticTrackedKeys(FPCGSelectionKeyToSettingsMap& OutKeysToSettings, TArray<TObjectPtr<const UPCGGraph>>& OutVisitedGraphs) const
 {
-	if (Asset.IsNull())
+	if (Asset.IsNull() || IsPropertyOverriddenByPin(GET_MEMBER_NAME_CHECKED(UPCGLoadDataAssetSettings, Asset)))
 	{
 		return;
 	}
diff --git a/Engine/Plugins/PCG/Source/PCG/Private/Elements/PCGBoundsFromMesh.cpp b/Engine/Plugins/PCG/Source/PCG/Private/Elements/PCGBoundsFromMesh.cpp
--- a/Engine/Plugins/PCG/Source/PCG/Private/Elements/PCGBoundsFromMesh.cpp
+++ b/Engine/Plugins/PCG/Source/PCG/Private/Elements/PCGBoundsFromMesh.cpp
@@ -10,6 +10,10 @@
 #include "Metadata/Accessors/PCGAttributeAccessorHelpers.h"
 #include "Metadata/Accessors/PCGAttributeAccessorKeys.h"
 
+#if WITH_EDITOR
+#include "Helpers/PCGDynamicTrackingHelpers.h"
+#endif // WITH_EDITOR
+
 #include "Algo/Transform.h"
 #include "Engine/StaticMesh.h"
 
@@ -220,6 +224,12 @@ bool FPCGBoundsFromMeshElement::ExecuteInternal(FPCGContext* InContext) const
 
 	if (!Context->bBoundsQueried)
 	{
+#if WITH_EDITOR
+		// Tracking is always dynamic in this case, since we are always pulling the mesh from an attribute.
+		FPCGDynamicTrackingHelper DynamicTracking;
+		DynamicTracking.EnableAndInitialize(Context, Context->MeshesToLoad.Num());
+#endif // WITH_EDITOR
+		
 		for (int MeshIndex = 0; MeshIndex < Context->MeshesToLoad.Num(); ++MeshIndex)
 		{
 			const FSoftObjectPath& MeshPath = Context->MeshesToLoad[MeshIndex];
@@ -228,7 +238,15 @@ bool FPCGBoundsFromMeshElement::ExecuteInternal(FPCGContext* InContext) const
 			if (Mesh.IsValid() && Mesh.Get())
 			{
 				Context->MeshToBoundsMap.Add(MeshIndex, Mesh->GetBoundingBox());
+
+#if WITH_EDITOR
+				DynamicTracking.AddToTracking(FPCGSelectionKey::CreateFromPath(MeshPath), /*bIsCulled=*/ false);
+#endif // WITH_EDITOR
 			}
+
+#if WITH_EDITOR
+			DynamicTracking.Finalize(Context);
+#endif // WITH_EDITOR
 		}
 
 		Context->bBoundsQueried = true;
diff --git a/Engine/Plugins/PCG/Source/PCG/Private/Elements/PCGGetStaticMeshResourceData.cpp b/Engine/Plugins/PCG/Source/PCG/Private/Elements/PCGGetStaticMeshResourceData.cpp
--- a/Engine/Plugins/PCG/Source/PCG/Private/Elements/PCGGetStaticMeshResourceData.cpp
+++ b/Engine/Plugins/PCG/Source/PCG/Private/Elements/PCGGetStaticMeshResourceData.cpp
@@ -6,6 +6,10 @@
 #include "PCGModule.h"
 #include "Metadata/Accessors/PCGAttributeAccessorHelpers.h"
 
+#if WITH_EDITOR
+#include "Helpers/PCGDynamicTrackingHelpers.h"
+#endif // WITH_EDITOR
+
 #include "Engine/StaticMesh.h"
 
 #include UE_INLINE_GENERATED_CPP_BY_NAME(PCGGetStaticMeshResourceData)
@@ -48,6 +52,20 @@ FPCGElementPtr UPCGGetStaticMeshResourceDataSettings::CreateElement() const
 }
 
 #if WITH_EDITOR
+void UPCGGetStaticMeshResourceDataSettings::GetStaticTrackedKeys(FPCGSelectionKeyToSettingsMap& OutKeysToSettings, TArray<TObjectPtr<const UPCGGraph>>& OutVisitedGraphs) const
+{
+	if (bOverrideFromInput)
+	{
+		return;
+	}
+	
+	for (const TSoftObjectPtr<UStaticMesh>& StaticMesh : StaticMeshes)
+	{
+		FPCGSelectionKey Key = FPCGSelectionKey::CreateFromPath(StaticMesh.ToSoftObjectPath());
+		OutKeysToSettings.FindOrAdd(MoveTemp(Key)).Emplace(this, /*bCulling=*/false);
+	}
+}
+
 EPCGChangeType UPCGGetStaticMeshResourceDataSettings::GetChangeTypeForProperty(const FName& InPropertyName) const
 {
 	EPCGChangeType ChangeType = Super::GetChangeTypeForProperty(InPropertyName);
@@ -105,12 +123,24 @@ bool FPCGGetStaticMeshResourceDataElement::ExecuteInternal(FPCGContext* InContex
 	// We only want to create one resource data per unique mesh, so we'll cache them here while looping over the input.
 	TMap<TSoftObjectPtr<UStaticMesh>, UPCGStaticMeshResourceData*> CreatedResourceDatas;
 
+#if WITH_EDITOR
+	FPCGDynamicTrackingHelper DynamicTracking;
+	if (Settings->bOverrideFromInput)
+	{
+		DynamicTracking.EnableAndInitialize(InContext, StaticMeshes.Num());
+	}
+#endif // WITH_EDITOR
+
 	for (const TSoftObjectPtr<UStaticMesh>& StaticMesh : StaticMeshes)
 	{
 		if (StaticMesh.IsNull())
 		{
 			continue;
 		}
+
+#if WITH_EDITOR
+		DynamicTracking.AddToTracking(FPCGSelectionKey::CreateFromPath(StaticMesh.ToSoftObjectPath()), /*bCulled=*/false);
+#endif // WITH_EDITOR
 		
 		UPCGStaticMeshResourceData** StaticMeshResourceDataPtr = CreatedResourceDatas.Find(StaticMesh);
 		UPCGStaticMeshResourceData* StaticMeshResourceData = nullptr;
@@ -130,6 +160,10 @@ bool FPCGGetStaticMeshResourceDataElement::ExecuteInternal(FPCGContext* InContex
 		InContext->OutputData.TaggedData.Emplace_GetRef().Data = StaticMeshResourceData;
 	}
 
+#if WITH_EDITOR
+	DynamicTracking.Finalize(InContext);
+#endif // WITH_EDITOR
+
 	return true;
 }
 
diff --git a/Engine/Plugins/PCG/Source/PCG/Private/Elements/PCGPointFromMeshElement.cpp b/Engine/Plugins/PCG/Source/PCG/Private/Elements/PCGPointFromMeshElement.cpp
--- a/Engine/Plugins/PCG/Source/PCG/Private/Elements/PCGPointFromMeshElement.cpp
+++ b/Engine/Plugins/PCG/Source/PCG/Private/Elements/PCGPointFromMeshElement.cpp
@@ -7,6 +7,10 @@
 #include "Data/PCGPointData.h"
 #include "Data/PCGSpatialData.h"
 
+#if WITH_EDITOR
+#include "Helpers/PCGDynamicTrackingHelpers.h"
+#endif // WITH_EDITOR
+
 #include "Engine/StaticMesh.h"
 
 #include UE_INLINE_GENERATED_CPP_BY_NAME(PCGPointFromMeshElement)
@@ -23,6 +27,18 @@ FText UPCGPointFromMeshSettings::GetNodeTooltipText() const
 {
 	return LOCTEXT("PointFromMeshNodeTooltip", "Creates a single point at the origin with an attribute named MeshPathAttributeName containing a SoftObjectPath to the StaticMesh.");
 }
+
+void UPCGPointFromMeshSettings::GetStaticTrackedKeys(FPCGSelectionKeyToSettingsMap& OutKeysToSettings, TArray<TObjectPtr<const UPCGGraph>>& OutVisitedGraphs) const
+{
+	if (StaticMesh.IsNull() || IsPropertyOverriddenByPin(GET_MEMBER_NAME_CHECKED(UPCGPointFromMeshSettings, StaticMesh)))
+	{
+		return;
+	}
+
+	FPCGSelectionKey Key = FPCGSelectionKey::CreateFromPath(StaticMesh.ToSoftObjectPath());
+
+	OutKeysToSettings.FindOrAdd(MoveTemp(Key)).Emplace(this, /*bCulling=*/false);
+}
 #endif
 
 FPCGElementPtr UPCGPointFromMeshSettings::CreateElement() const
@@ -74,6 +90,13 @@ bool FPCGPointFromMeshElement::ExecuteInternal(FPCGContext* Context) const
 		return true;
 	}
 
+#if WITH_EDITOR
+	if (Context->IsValueOverriden(GET_MEMBER_NAME_CHECKED(UPCGPointFromMeshSettings, StaticMesh)))
+	{
+		FPCGDynamicTrackingHelper::AddSingleDynamicTrackingKey(Context, FPCGSelectionKey::CreateFromPath(Settings->StaticMesh.ToSoftObjectPath()), /*bIsCulled=*/false);
+	}
+#endif // WITH_EDITOR
+
 	TArray<FPCGTaggedData>& Outputs = Context->OutputData.TaggedData;
 	UPCGBasePointData* OutPointData = FPCGContext::NewPointData_AnyThread(Context);
 	Outputs.Emplace_GetRef().Data = OutPointData;
diff --git a/Engine/Plugins/PCG/Source/PCG/Public/Elements/PCGBoundsFromMesh.h b/Engine/Plugins/PCG/Source/PCG/Public/Elements/PCGBoundsFromMesh.h
--- a/Engine/Plugins/PCG/Source/PCG/Public/Elements/PCGBoundsFromMesh.h
+++ b/Engine/Plugins/PCG/Source/PCG/Public/Elements/PCGBoundsFromMesh.h
@@ -26,6 +26,7 @@ public:
 	virtual FText GetNodeTooltipText() const override;
 	virtual EPCGSettingsType GetType() const override { return EPCGSettingsType::Spatial; }
 	virtual bool IsPinUsedByNodeExecution(const UPCGPin* InPin) const override;
+	virtual bool CanDynamicallyTrackKeys() const override { return true; };
 #endif //WITH_EDITOR
 
 protected:
diff --git a/Engine/Plugins/PCG/Source/PCG/Public/Elements/PCGGetStaticMeshResourceData.h b/Engine/Plugins/PCG/Source/PCG/Public/Elements/PCGGetStaticMeshResourceData.h
--- a/Engine/Plugins/PCG/Source/PCG/Public/Elements/PCGGetStaticMeshResourceData.h
+++ b/Engine/Plugins/PCG/Source/PCG/Public/Elements/PCGGetStaticMeshResourceData.h
@@ -23,6 +23,8 @@ public:
 	virtual FName GetDefaultNodeName() const override { return FName(TEXT("GetStaticMeshResourceData")); }
 	virtual FText GetDefaultNodeTitle() const override { return NSLOCTEXT("PCGGetStaticMeshResourceDataElement", "NodeTitle", "Get Static Mesh Resource Data"); }
 	virtual EPCGSettingsType GetType() const override { return EPCGSettingsType::Resource; }
+	virtual void GetStaticTrackedKeys(FPCGSelectionKeyToSettingsMap& OutKeysToSettings, TArray<TObjectPtr<const UPCGGraph>>& OutVisitedGraphs) const override;
+	virtual bool CanDynamicallyTrackKeys() const override { return bOverrideFromInput; }
 #endif
 
 protected:
diff --git a/Engine/Plugins/PCG/Source/PCG/Public/Elements/PCGPointFromMeshElement.h b/Engine/Plugins/PCG/Source/PCG/Public/Elements/PCGPointFromMeshElement.h
--- a/Engine/Plugins/PCG/Source/PCG/Public/Elements/PCGPointFromMeshElement.h
+++ b/Engine/Plugins/PCG/Source/PCG/Public/Elements/PCGPointFromMeshElement.h
@@ -25,6 +25,8 @@ public:
 	virtual FText GetDefaultNodeTitle() const override;
 	virtual FText GetNodeTooltipText() const override;
 	virtual EPCGSettingsType GetType() const override { return EPCGSettingsType::Spatial; }
+	virtual void GetStaticTrackedKeys(FPCGSelectionKeyToSettingsMap& OutKeysToSettings, TArray<TObjectPtr<const UPCGGraph>>& OutVisitedGraphs) const override;
+	virtual bool CanDynamicallyTrackKeys() const override { return true; };
 #endif
 
 protected:
diff --git a/Engine/Plugins/PCGInterops/PCGGeometryScriptInterop/Source/PCGGeometryScriptInterop/Private/Elements/PCGAppendMeshesFromPoints.cpp b/Engine/Plugins/PCGInterops/PCGGeometryScriptInterop/Source/PCGGeometryScriptInterop/Private/Elements/PCGAppendMeshesFromPoints.cpp
--- a/Engine/Plugins/PCGInterops/PCGGeometryScriptInterop/Source/PCGGeometryScriptInterop/Private/Elements/PCGAppendMeshesFromPoints.cpp
+++ b/Engine/Plugins/PCGInterops/PCGGeometryScriptInterop/Source/PCGGeometryScriptInterop/Private/Elements/PCGAppendMeshesFromPoints.cpp
@@ -9,6 +9,10 @@
 #include "Helpers/PCGGeometryHelpers.h"
 #include "Metadata/Accessors/PCGAttributeAccessorHelpers.h"
 
+#if WITH_EDITOR
+#include "Helpers/PCGDynamicTrackingHelpers.h"
+#endif // WITH_EDITOR
+
 #include "DynamicMeshEditor.h"
 #include "UDynamicMesh.h"
 #include "ConversionUtils/SceneComponentToDynamicMesh.h"
@@ -41,6 +45,18 @@ FText UPCGAppendMeshesFromPointsSettings::GetNodeTooltipText() const
 	return LOCTEXT("NodeTooltip", "Append meshes at the points transforms. Mesh can be a single static mesh, multiple meshes coming from the points or another dynamic mesh.");
 }
 
+void UPCGAppendMeshesFromPointsSettings::GetStaticTrackedKeys(FPCGSelectionKeyToSettingsMap& OutKeysToSettings, TArray<TObjectPtr<const UPCGGraph>>& OutVisitedGraphs) const
+{
+	if (Mode != EPCGAppendMeshesFromPointsMode::SingleStaticMesh || StaticMesh.IsNull() || IsPropertyOverriddenByPin(GET_MEMBER_NAME_CHECKED(UPCGAppendMeshesFromPointsSettings, StaticMesh)))
+	{
+		return;
+	}
+
+	FPCGSelectionKey Key = FPCGSelectionKey::CreateFromPath(StaticMesh.ToSoftObjectPath());
+
+	OutKeysToSettings.FindOrAdd(MoveTemp(Key)).Emplace(this, /*bCulling=*/false);
+}
+
 EPCGChangeType UPCGAppendMeshesFromPointsSettings::GetChangeTypeForProperty(const FName& InPropertyName) const
 {
 	EPCGChangeType ChangeType = Super::GetChangeTypeForProperty(InPropertyName) | EPCGChangeType::Cosmetic;
@@ -247,10 +263,22 @@ bool FPCGAppendMeshesFromPointsElement::ExecuteInternal(FPCGContext* InContext)
 			return true;
 		}
 
+#if WITH_EDITOR
+		if (Context->IsValueOverriden(GET_MEMBER_NAME_CHECKED(UPCGAppendMeshesFromPointsSettings, StaticMesh)))
+		{
+			FPCGDynamicTrackingHelper::AddSingleDynamicTrackingKey(Context, FPCGSelectionKey::CreateFromPath(Settings->StaticMesh.ToSoftObjectPath()), /*bCulled=*/false);
+		}
+#endif // WITH_EDITOR
+
 		break;
 	}
 	case EPCGAppendMeshesFromPointsMode::StaticMeshFromAttribute:
 	{
+#if WITH_EDITOR
+		FPCGDynamicTrackingHelper DynamicTracking;
+		DynamicTracking.EnableAndInitialize(Context, Context->MeshToPointIndicesMapping.Num());
+#endif // WITH_EDITOR
+
 		StaticMeshToDynMesh.Reserve(Context->MeshToPointIndicesMapping.Num());
 	
 		for (TPair<FSoftObjectPath, TArray<int32>>& It : Context->MeshToPointIndicesMapping)
@@ -259,8 +287,16 @@ bool FPCGAppendMeshesFromPointsElement::ExecuteInternal(FPCGContext* InContext)
 			{
 				return true;
 			}
+
+#if WITH_EDITOR
+		DynamicTracking.AddToTracking(FPCGSelectionKey::CreateFromPath(It.Key), /*bCulled=*/false);
+#endif // WITH_EDITOR
 		}
 
+#if WITH_EDITOR
+		DynamicTracking.Finalize(Context);
+#endif // WITH_EDITOR
+
 		break;
 	}
 	case EPCGAppendMeshesFromPointsMode::DynamicMesh:
diff --git a/Engine/Plugins/PCGInterops/PCGGeometryScriptInterop/Source/PCGGeometryScriptInterop/Private/Elements/PCGMeshSampler.cpp b/Engine/Plugins/PCGInterops/PCGGeometryScriptInterop/Source/PCGGeometryScriptInterop/Private/Elements/PCGMeshSampler.cpp
--- a/Engine/Plugins/PCGInterops/PCGGeometryScriptInterop/Source/PCGGeometryScriptInterop/Private/Elements/PCGMeshSampler.cpp
+++ b/Engine/Plugins/PCGInterops/PCGGeometryScriptInterop/Source/PCGGeometryScriptInterop/Private/Elements/PCGMeshSampler.cpp
@@ -404,7 +404,7 @@ FText UPCGMeshSamplerSettings::GetNodeTooltipText() const
 
 void UPCGMeshSamplerSettings::GetStaticTrackedKeys(FPCGSelectionKeyToSettingsMap& OutKeysToSettings, TArray<TObjectPtr<const UPCGGraph>>& OutVisitedGraphs) const
 {
-	if (StaticMesh.IsNull())
+	if (StaticMesh.IsNull() || IsPropertyOverriddenByPin(GET_MEMBER_NAME_CHECKED(UPCGMeshSamplerSettings, StaticMesh)))
 	{
 		return;
 	}
diff --git a/Engine/Plugins/PCGInterops/PCGGeometryScriptInterop/Source/PCGGeometryScriptInterop/Private/Elements/PCGStaticMeshToDynamicMeshElement.cpp b/Engine/Plugins/PCGInterops/PCGGeometryScriptInterop/Source/PCGGeometryScriptInterop/Private/Elements/PCGStaticMeshToDynamicMeshElement.cpp
--- a/Engine/Plugins/PCGInterops/PCGGeometryScriptInterop/Source/PCGGeometryScriptInterop/Private/Elements/PCGStaticMeshToDynamicMeshElement.cpp
+++ b/Engine/Plugins/PCGInterops/PCGGeometryScriptInterop/Source/PCGGeometryScriptInterop/Private/Elements/PCGStaticMeshToDynamicMeshElement.cpp
@@ -6,6 +6,10 @@
 #include "PCGModule.h"
 #include "Data/PCGDynamicMeshData.h"
 
+#if WITH_EDITOR
+#include "Helpers/PCGDynamicTrackingHelpers.h"
+#endif // WITH_EDITOR
+
 #include "UDynamicMesh.h"
 #include "ConversionUtils/SceneComponentToDynamicMesh.h"
 #include "Engine/StaticMesh.h"
@@ -28,6 +32,18 @@ FText UPCGStaticMeshToDynamicMeshSettings::GetNodeTooltipText() const
 {
 	return LOCTEXT("NodeTooltip", "Convert a static mesh into a dynamic mesh data.");
 }
+
+void UPCGStaticMeshToDynamicMeshSettings::GetStaticTrackedKeys(FPCGSelectionKeyToSettingsMap& OutKeysToSettings, TArray<TObjectPtr<const UPCGGraph>>& OutVisitedGraphs) const
+{
+	if (StaticMesh.IsNull() || IsPropertyOverriddenByPin(GET_MEMBER_NAME_CHECKED(UPCGStaticMeshToDynamicMeshSettings, StaticMesh)))
+	{
+		return;
+	}
+
+	FPCGSelectionKey Key = FPCGSelectionKey::CreateFromPath(StaticMesh.ToSoftObjectPath());
+
+	OutKeysToSettings.FindOrAdd(MoveTemp(Key)).Emplace(this, /*bCulling=*/false);
+}
 #endif // WITH_EDITOR
 
 FPCGElementPtr UPCGStaticMeshToDynamicMeshSettings::CreateElement() const
@@ -104,6 +120,13 @@ bool FPCGStaticMeshToDynamicMeshElement::ExecuteInternal(FPCGContext* InContext)
 		return true;
 	}
 
+#if WITH_EDITOR
+	if (InContext->IsValueOverriden(GET_MEMBER_NAME_CHECKED(UPCGStaticMeshToDynamicMeshSettings, StaticMesh)))
+	{
+		FPCGDynamicTrackingHelper::AddSingleDynamicTrackingKey(InContext, FPCGSelectionKey::CreateFromPath(Settings->StaticMesh.ToSoftObjectPath()), /*bIsCulled=*/false);
+	}
+#endif // WITH_EDITOR
+
 	if (Settings->bExtractMaterials && !Settings->OverrideMaterials.IsEmpty())
 	{
 		if (StaticMesh->GetStaticMaterials().Num() != Settings->OverrideMaterials.Num())
diff --git a/Engine/Plugins/PCGInterops/PCGGeometryScriptInterop/Source/PCGGeometryScriptInterop/Public/Elements/PCGAppendMeshesFromPoints.h b/Engine/Plugins/PCGInterops/PCGGeometryScriptInterop/Source/PCGGeometryScriptInterop/Public/Elements/PCGAppendMeshesFromPoints.h
--- a/Engine/Plugins/PCGInterops/PCGGeometryScriptInterop/Source/PCGGeometryScriptInterop/Public/Elements/PCGAppendMeshesFromPoints.h
+++ b/Engine/Plugins/PCGInterops/PCGGeometryScriptInterop/Source/PCGGeometryScriptInterop/Public/Elements/PCGAppendMeshesFromPoints.h
@@ -32,6 +32,8 @@ public:
 	virtual FName GetDefaultNodeName() const override;
 	virtual FText GetDefaultNodeTitle() const override;
 	virtual FText GetNodeTooltipText() const override;
+	virtual void GetStaticTrackedKeys(FPCGSelectionKeyToSettingsMap& OutKeysToSettings, TArray<TObjectPtr<const UPCGGraph>>& OutVisitedGraphs) const override;
+	virtual bool CanDynamicallyTrackKeys() const override { return Mode != EPCGAppendMeshesFromPointsMode::DynamicMesh; };
 #endif
 
 protected:
diff --git a/Engine/Plugins/PCGInterops/PCGGeometryScriptInterop/Source/PCGGeometryScriptInterop/Public/Elements/PCGStaticMeshToDynamicMeshElement.h b/Engine/Plugins/PCGInterops/PCGGeometryScriptInterop/Source/PCGGeometryScriptInterop/Public/Elements/PCGStaticMeshToDynamicMeshElement.h
--- a/Engine/Plugins/PCGInterops/PCGGeometryScriptInterop/Source/PCGGeometryScriptInterop/Public/Elements/PCGStaticMeshToDynamicMeshElement.h
+++ b/Engine/Plugins/PCGInterops/PCGGeometryScriptInterop/Source/PCGGeometryScriptInterop/Public/Elements/PCGStaticMeshToDynamicMeshElement.h
@@ -26,6 +26,8 @@ public:
 	virtual FName GetDefaultNodeName() const override;
 	virtual FText GetDefaultNodeTitle() const override;
 	virtual FText GetNodeTooltipText() const override;
+	virtual void GetStaticTrackedKeys(FPCGSelectionKeyToSettingsMap& OutKeysToSettings, TArray<TObjectPtr<const UPCGGraph>>& OutVisitedGraphs) const override;
+	virtual bool CanDynamicallyTrackKeys() const override { return true; };
 #endif
 
 protected:
