[PCG] Added support for collision wrapper data to be used in primitive cross section node. Backend work to support more types to dynmesh.

#rnx

diff --git a/Engine/Plugins/PCG/Source/PCG/Public/Data/PCGCollisionShapeData.h b/Engine/Plugins/PCG/Source/PCG/Public/Data/PCGCollisionShapeData.h
--- a/Engine/Plugins/PCG/Source/PCG/Public/Data/PCGCollisionShapeData.h
+++ b/Engine/Plugins/PCG/Source/PCG/Public/Data/PCGCollisionShapeData.h
@@ -45,6 +45,9 @@ public:
 	virtual const UPCGPointArrayData* CreatePointArrayData(FPCGContext* Context, const FBox& InBounds) const override;
 	// ~End UPCGSpatialDataWithPointCache implementation
 
+	const FTransform& GetTransform() const { return Transform; }
+	const FCollisionShape& GetCollisionShape() const { return Shape; }
+
 private:
 	virtual const UPCGBasePointData* CreateBasePointData(FPCGContext* Context, TSubclassOf<UPCGBasePointData> PointDataClass) const;
 
diff --git a/Engine/Plugins/PCG/Source/PCG/Public/Data/PCGCollisionWrapperData.h b/Engine/Plugins/PCG/Source/PCG/Public/Data/PCGCollisionWrapperData.h
--- a/Engine/Plugins/PCG/Source/PCG/Public/Data/PCGCollisionWrapperData.h
+++ b/Engine/Plugins/PCG/Source/PCG/Public/Data/PCGCollisionWrapperData.h
@@ -88,6 +88,7 @@ protected:
 	virtual UPCGSpatialData* CopyInternal(FPCGContext* Context) const override;
 	// ~End UPCGSpatialData interface
 
+public:
 	// For performance reasons, we keep a raw pointer to the point data in editor.
 #if WITH_EDITOR
 	inline const UPCGBasePointData* GetPointData() const { return RawPointData; }
@@ -95,6 +96,10 @@ protected:
 	inline const UPCGBasePointData* GetPointData() const { return PointData.Get(); }
 #endif
 
+	// Some helper functions when this is used in different contexts
+	EPCGCollisionQueryFlag GetCollisionFlag() const { return CollisionQueryFlag; }
+	const FPCGCollisionWrapper& GetCollisionWrapper() const { return CollisionWrapper; }
+
 private:
 	const PhysicsInterfaceTypes::FInlineShapeArray& GetCachedShapes(int32 EntryIndex) const;
 
diff --git a/Engine/Plugins/PCG/Source/PCG/Public/Data/PCGVolumeData.h b/Engine/Plugins/PCG/Source/PCG/Public/Data/PCGVolumeData.h
--- a/Engine/Plugins/PCG/Source/PCG/Public/Data/PCGVolumeData.h
+++ b/Engine/Plugins/PCG/Source/PCG/Public/Data/PCGVolumeData.h
@@ -45,6 +45,8 @@ public:
 	UE_API virtual const UPCGPointArrayData* CreatePointArrayData(FPCGContext* Context, const FBox& InBounds) const override;
 	// ~End UPCGSpatialDataWithPointCache interface
 
+	TWeakObjectPtr<AVolume> GetVolumeActor() const { return Volume; }
+
 	UPROPERTY(BlueprintReadWrite, EditAnywhere, Category = "Data")
 	FVector VoxelSize = FVector(100.0, 100.0, 100.0);
 
diff --git a/Engine/Plugins/PCGInterops/PCGGeometryScriptInterop/Source/PCGGeometryScriptInterop/Private/Elements/PCGPrimitiveCrossSection.cpp b/Engine/Plugins/PCGInterops/PCGGeometryScriptInterop/Source/PCGGeometryScriptInterop/Private/Elements/PCGPrimitiveCrossSection.cpp
--- a/Engine/Plugins/PCGInterops/PCGGeometryScriptInterop/Source/PCGGeometryScriptInterop/Private/Elements/PCGPrimitiveCrossSection.cpp
+++ b/Engine/Plugins/PCGInterops/PCGGeometryScriptInterop/Source/PCGGeometryScriptInterop/Private/Elements/PCGPrimitiveCrossSection.cpp
@@ -28,7 +28,6 @@ namespace PCGPrimitiveCrossSection::Constants
 {
 	const FName DefaultExtrusionVectorAttributeName(TEXT("ExtrusionVector"));
 	static constexpr FGeometryScriptMeshPlaneCutOptions CutPlaneOptions{.bFillHoles = false, .bFillSpans = false, .bFlipCutSide = true};
-	static constexpr FGeometryScriptCopyMeshFromComponentOptions CopyMeshFromComponentOptions{.bWantNormals = false, .bWantTangents = false};
 }
 
 struct FCrossSection
@@ -74,24 +73,7 @@ bool FPCGPrimitiveCrossSectionElement::ExecuteInternal(FPCGContext* Context) con
 		return true;
 	}
 
-	// TODO: Could also take in PCGVolume
-	TArray<const UPCGPrimitiveData*> PrimitiveDataCollection;
-	PrimitiveDataCollection.Reserve(PrimitiveInputs.Num());
-
-	for (const FPCGTaggedData& TaggedData : PrimitiveInputs)
-	{
-		if (const UPCGPrimitiveData* PrimitiveData = Cast<UPCGPrimitiveData>(TaggedData.Data))
-		{
-			PrimitiveDataCollection.Add(PrimitiveData);
-		}
-	}
-
-	if (PrimitiveDataCollection.IsEmpty())
-	{
-		return true;
-	}
-
-	TArray<UDynamicMesh*, TInlineAllocator<16>> DynamicMeshes;
+	TArray<UDynamicMesh*> DynamicMeshes;
 #if WITH_EDITOR
 	UGeometryScriptDebug* DynamicMeshDebug = FPCGContext::NewObject_AnyThread<UGeometryScriptDebug>(Context);
 #else
@@ -99,29 +81,14 @@ bool FPCGPrimitiveCrossSectionElement::ExecuteInternal(FPCGContext* Context) con
 #endif // WITH_EDITOR
 
 	// Collect all the primitives and append them to the dynamic mesh
-	for (const UPCGPrimitiveData* PrimitiveData : PrimitiveDataCollection)
+	for (const FPCGTaggedData& TaggedData : PrimitiveInputs)
 	{
-		TWeakObjectPtr<UPrimitiveComponent> PrimitivePtr = PrimitiveData->GetComponent();
-		if (!PrimitivePtr.IsValid())
-		{
-			continue;
-		}
+		PCGGeometryHelpers::ConvertDataToDynMeshes(TaggedData.Data, Context, DynamicMeshes, /*bMergeMeshes=*/false, DynamicMeshDebug);
+	}
 
-		// Convert from scene component to mesh and begin boolean operation
-		if (USceneComponent* SceneComponent = Cast<USceneComponent>(PrimitivePtr.Pin().Get()))
-		{
-			EGeometryScriptOutcomePins Outcome;
-			FTransform PrimitiveTransform;
-			// TODO: If we have repeats, ie. components that match content/transform, we could skip them.
-			UGeometryScriptLibrary_SceneUtilityFunctions::CopyMeshFromComponent(
-				SceneComponent,
-				DynamicMeshes.Emplace_GetRef(FPCGContext::NewObject_AnyThread<UDynamicMesh>(Context)),
-				PCGPrimitiveCrossSection::Constants::CopyMeshFromComponentOptions,
-				/*bTransformToWorld=*/true,
-				PrimitiveTransform,
-				Outcome,
-				DynamicMeshDebug);
-		}
+	if (DynamicMeshes.IsEmpty())
+	{
+		return true;
 	}
 
 	// Dissolve the meshes down until there are none intersecting.
@@ -219,7 +186,7 @@ bool FPCGPrimitiveCrossSectionElement::ExecuteInternal(FPCGContext* Context) con
 
 		int32 LastGroupIndex = 0;
 		int32 CurrentIndex = 1;
-		while (CurrentIndex < ProjectionScalars.Num() - 1)
+		while (CurrentIndex < ProjectionScalars.Num())
 		{
 			// Consecutive value falls within the threshold. Merge them.
 			if (ProjectionScalars[CurrentIndex] - ProjectionScalars[LastGroupIndex] >= MergingThreshold)
@@ -247,7 +214,7 @@ bool FPCGPrimitiveCrossSectionElement::ExecuteInternal(FPCGContext* Context) con
 		SlicePlaneTransform.SetRotation(AdjustedRotation);
 
 		TArray<FCrossSection> CrossSections;
-		for (int TierIndex = 0; TierIndex < TierHeights.Num() - 1; ++TierIndex)
+		for (int TierIndex = 0; TierIndex < TierHeights.Num(); ++TierIndex)
 		{
 			FVector SliceLocation = SliceOrigin + SliceDirection * TierHeights[TierIndex];
 			SlicePlaneTransform.SetLocation(SliceLocation);
@@ -256,9 +223,8 @@ bool FPCGPrimitiveCrossSectionElement::ExecuteInternal(FPCGContext* Context) con
 			UGeometryScriptLibrary_MeshBooleanFunctions::ApplyMeshPlaneCut(DynamicMesh, SlicePlaneTransform, PCGPrimitiveCrossSection::Constants::CutPlaneOptions, DynamicMeshDebug);
 
 			// Break now as there is no point in continuing along this invalid mesh
-			if (!ensure(DynamicMesh->GetTriangleCount() >= 1))
+			if (DynamicMesh->GetTriangleCount() == 0)
 			{
-				PCGLog::LogWarningOnGraph(FText::Format(LOCTEXT("DynamicMeshInvalid", "Dynamic Mesh Invalid at tier: {0}"), FText::AsNumber(TierIndex)), Context);
 				break;
 			}
 
diff --git a/Engine/Plugins/PCGInterops/PCGGeometryScriptInterop/Source/PCGGeometryScriptInterop/Private/Helpers/PCGGeometryHelpers.cpp b/Engine/Plugins/PCGInterops/PCGGeometryScriptInterop/Source/PCGGeometryScriptInterop/Private/Helpers/PCGGeometryHelpers.cpp
--- a/Engine/Plugins/PCGInterops/PCGGeometryScriptInterop/Source/PCGGeometryScriptInterop/Private/Helpers/PCGGeometryHelpers.cpp
+++ b/Engine/Plugins/PCGInterops/PCGGeometryScriptInterop/Source/PCGGeometryScriptInterop/Private/Helpers/PCGGeometryHelpers.cpp
@@ -3,11 +3,37 @@
 #include "Helpers/PCGGeometryHelpers.h"
 
 #include "ConversionUtils/SceneComponentToDynamicMesh.h"
+#include "Data/PCGCollisionShapeData.h"
+#include "Data/PCGCollisionWrapperData.h"
+#include "Data/PCGPrimitiveData.h"
+#include "Data/PCGVolumeData.h"
 #include "DynamicMesh/DynamicMeshAttributeSet.h"
 #include "Utils/PCGLogErrors.h"
 
+#include "UDynamicMesh.h"
+#include "TransformSequence.h"
+#include "Components/BrushComponent.h"
+#include "DynamicMesh/MeshTransforms.h"
+#include "GameFramework/Volume.h"
+#include "Generators/CapsuleGenerator.h"
+#include "Generators/MinimalBoxMeshGenerator.h"
+#include "Generators/SphereGenerator.h"
+#include "GeometryScript/SceneUtilityFunctions.h"
+#include "Physics/ComponentCollisionUtil.h"
+#include "PhysicsEngine/BodyInstance.h"
+#include "PhysicsEngine/BodySetup.h"
+
 namespace PCGGeometryHelpers
 {
+	namespace Constants
+	{
+		static constexpr FGeometryScriptCopyMeshFromComponentOptions CopyMeshFromComponentOptions{ .bWantNormals = false, .bWantTangents = false };
+		static constexpr int32 SphereNumTheta = 25;
+		static constexpr int32 SphereNumPhi = 25;
+		static constexpr int32 CapsuleNumHemisphereSteps = 12;
+		static constexpr int32 CapsuleNumCircleSteps = 12;
+	}
+
 	void GeometryScriptDebugToPCGLog(FPCGContext* Context, const UGeometryScriptDebug* Debug)
 	{
 		check(Context && Debug);
@@ -107,4 +133,222 @@ namespace PCGGeometryHelpers
 			}
 		}
 	}
-}
\ No newline at end of file
+
+	bool ConvertDataToDynMeshes(const UPCGData* InData, FPCGContext* Context, TArray<UDynamicMesh*>& OutMeshes, bool bMergeMeshes, UGeometryScriptDebug* DynamicMeshDebug)
+	{
+		if (!InData)
+		{
+			return false;
+		}
+
+		auto CreateDynMeshFromGenerator = [&OutMeshes, Context](auto& Generator, const FTransform& Transform)
+		{
+			UDynamicMesh* NewMesh = OutMeshes.Emplace_GetRef(FPCGContext::NewObject_AnyThread<UDynamicMesh>(Context));
+			FDynamicMesh3& DynMesh = NewMesh->GetMeshRef();
+			DynMesh.Copy(&Generator);
+			MeshTransforms::ApplyTransform(DynMesh, Transform, /*bReverseOrientationIfNeeded=*/true);
+		};
+
+		auto CreateBoxDynMesh = [](const FBox& Box, const FTransform& Transform, FDynamicMesh3& OutMesh)
+		{
+			UE::Geometry::FMinimalBoxMeshGenerator Generator;
+			UE::Geometry::FOrientedBox3d OrientedBox;
+			OrientedBox.Frame = UE::Geometry::FFrame3d(Box.GetCenter());
+			OrientedBox.Extents = Box.GetExtent();
+			Generator.Box = OrientedBox;
+			
+			Generator.Generate();
+
+			OutMesh.Copy(&Generator);
+			MeshTransforms::ApplyTransform(OutMesh, Transform, /*bReverseOrientationIfNeeded=*/true);
+		};
+
+		auto AddBoxDynMesh = [&CreateBoxDynMesh, &OutMeshes, Context](const FBox& Box, const FTransform& Transform)
+		{
+			UDynamicMesh* NewMesh = OutMeshes.Emplace_GetRef(FPCGContext::NewObject_AnyThread<UDynamicMesh>(Context));
+			FDynamicMesh3& DynMesh = NewMesh->GetMeshRef();
+			CreateBoxDynMesh(Box, Transform, DynMesh);
+		};
+
+		if (const UPCGPrimitiveData* PrimitiveData = Cast<UPCGPrimitiveData>(InData))
+		{
+			TWeakObjectPtr<UPrimitiveComponent> PrimitivePtr = PrimitiveData->GetComponent();
+			if (!PrimitivePtr.IsValid())
+			{
+				return false;
+			}
+
+			// Convert from scene component to mesh and begin boolean operation
+			if (USceneComponent* SceneComponent = Cast<USceneComponent>(PrimitivePtr.Pin().Get()))
+			{
+				EGeometryScriptOutcomePins Outcome;
+				FTransform PrimitiveTransform;
+				// TODO: If we have repeats, ie. components that match content/transform, we could skip them.
+				UGeometryScriptLibrary_SceneUtilityFunctions::CopyMeshFromComponent(
+					SceneComponent,
+					OutMeshes.Emplace_GetRef(FPCGContext::NewObject_AnyThread<UDynamicMesh>(Context)),
+					PCGGeometryHelpers::Constants::CopyMeshFromComponentOptions,
+					/*bTransformToWorld=*/true,
+					PrimitiveTransform,
+					Outcome,
+					DynamicMeshDebug);
+
+				if (Outcome != EGeometryScriptOutcomePins::Success)
+				{
+					OutMeshes.Pop();
+					return false;
+				}
+			}
+		}
+		else if (const UPCGVolumeData* VolumeData = Cast<UPCGVolumeData>(InData))
+		{
+			// There are two cases here: either with the brush component or otherwise the simple volume box
+			if (TStrongObjectPtr<AVolume> VolumeActor = VolumeData->GetVolumeActor().Pin())
+			{
+				EGeometryScriptOutcomePins Outcome;
+				FTransform PrimitiveTransform;
+
+				UGeometryScriptLibrary_SceneUtilityFunctions::CopyCollisionMeshesFromObject(
+					VolumeActor->GetBrushComponent(),
+					OutMeshes.Emplace_GetRef(FPCGContext::NewObject_AnyThread<UDynamicMesh>(Context)),
+					/*bTransformToWorld=*/true,
+					PrimitiveTransform,
+					Outcome,
+					/*bUseComplexCollision=*/false,
+					/*SphereResolution*/PCGGeometryHelpers::Constants::SphereNumPhi,
+					DynamicMeshDebug);
+
+				if (Outcome != EGeometryScriptOutcomePins::Success)
+				{
+					OutMeshes.Pop();
+					return false;
+				}
+			}
+			else
+			{
+				AddBoxDynMesh(VolumeData->GetBounds(), FTransform::Identity);
+			}
+		}
+		else if (const UPCGCollisionShapeData* ShapeData = Cast<UPCGCollisionShapeData>(InData))
+		{
+			const FTransform& ShapeTransform = ShapeData->GetTransform();
+			const FCollisionShape& CollisionShape = ShapeData->GetCollisionShape();
+
+			// Similar to FChaosVDGeometryBuilder::CreateMeshGeneratorForImplicitObject and FMakeSphereMeshDataflowNode::Evaluate
+			if (CollisionShape.IsBox())
+			{
+				FBox CollisionShapeBox = FBox::BuildAABB(FVector::ZeroVector, CollisionShape.GetBox());
+				AddBoxDynMesh(CollisionShapeBox, ShapeTransform);
+			}
+			else if (CollisionShape.IsSphere())
+			{
+				UE::Geometry::FSphereGenerator Generator;
+				Generator.Radius = CollisionShape.GetSphereRadius();
+				Generator.NumTheta = PCGGeometryHelpers::Constants::SphereNumTheta;
+				Generator.NumPhi = PCGGeometryHelpers::Constants::SphereNumPhi;
+				Generator.bPolygroupPerQuad = false;
+
+				Generator.Generate();
+				CreateDynMeshFromGenerator(Generator, ShapeTransform);
+			}
+			else if (CollisionShape.IsCapsule())
+			{
+				UE::Geometry::FCapsuleGenerator Generator;
+				Generator.Radius = CollisionShape.GetCapsuleRadius();
+				Generator.SegmentLength = 2.0 * CollisionShape.GetCapsuleAxisHalfLength();
+				Generator.NumHemisphereArcSteps = PCGGeometryHelpers::Constants::CapsuleNumHemisphereSteps;
+				Generator.NumCircleSteps = PCGGeometryHelpers::Constants::CapsuleNumCircleSteps;
+				Generator.bPolygroupPerQuad = false;
+
+				Generator.Generate();
+				CreateDynMeshFromGenerator(Generator, ShapeTransform);
+			}
+			else
+			{
+				return false;
+			}
+		}
+		else if (const UPCGCollisionWrapperData* CollisionWrapperData = Cast<UPCGCollisionWrapperData>(InData))
+		{
+			const EPCGCollisionQueryFlag CollisionQueryFlag = CollisionWrapperData->GetCollisionFlag();
+			const FPCGCollisionWrapper& CollisionWrapper = CollisionWrapperData->GetCollisionWrapper();
+			const UPCGBasePointData* PointData = CollisionWrapperData->GetPointData();
+			check(PointData);
+			const TConstPCGValueRange<FTransform> PointTransforms = PointData->GetConstTransformValueRange();
+			const TConstPCGValueRange<FVector> PointMinBounds = PointData->GetConstBoundsMinValueRange();
+			const TConstPCGValueRange<FVector> PointMaxBounds = PointData->GetConstBoundsMaxValueRange();
+
+			UDynamicMesh* CurrentMesh = nullptr;
+
+			// @todo_pcg: we probably should create only one FDynamicMesh per unique body instance, then do transformations.
+			for (int PointIndex = 0; PointIndex < PointTransforms.Num(); ++PointIndex)
+			{
+				const FTransform& PointTransform = PointTransforms[PointIndex];
+				const FBodyInstance* BodyInstance = CollisionWrapper.GetBodyInstance(PointIndex);
+
+				// Implementation note: this is an adaptation of UGeometryScriptLibrary_SceneUtilityFunctions::CopyCollisionMeshesFromObject
+				FDynamicMesh3 AccumulatedMesh;
+				bool bHasData = false;
+
+				if (BodyInstance)
+				{
+					const UBodySetup* BodySetup = BodyInstance->GetBodySetup();
+					if (!BodySetup)
+					{
+						continue;
+					}
+
+					const int SphereResolution = 16;
+					FVector ExternalScale = PointTransform.GetScale3D();
+
+					UE::Geometry::FTransformSequence3d Transforms;
+					FTransform PointTransformNoScale = PointTransform;
+					PointTransformNoScale.SetScale3D(FVector::OneVector);
+					Transforms.Append(PointTransformNoScale);
+
+					// @todo_pcg: review this. Normally we should take the path that the shape is extracted from assumedly
+					if (CollisionQueryFlag == EPCGCollisionQueryFlag::Complex)
+					{
+						// Requires a IInterface_CollisionDataProvider - see UE::Geometry::ConvertComplexCollisionToMeshes
+					}
+					else
+					{
+						UE::Geometry::ConvertSimpleCollisionToMeshes(BodySetup->AggGeom, AccumulatedMesh, Transforms, SphereResolution, true, true, nullptr, false, ExternalScale);
+						bHasData = true;
+					}
+				}
+				else
+				{
+					// Consider it's a box then
+					FBox PointBox(PointMinBounds[PointIndex], PointMaxBounds[PointIndex]);
+					CreateBoxDynMesh(PointBox, PointTransform, AccumulatedMesh);
+					bHasData = true;
+				}
+
+				if (bHasData)
+				{
+					if (!CurrentMesh)
+					{
+						CurrentMesh = FPCGContext::NewObject_AnyThread<UDynamicMesh>(Context);
+						OutMeshes.Emplace(CurrentMesh);
+					}
+
+					if (bMergeMeshes)
+					{
+						CurrentMesh->EditMesh([&AccumulatedMesh](FDynamicMesh3& EditMesh)
+						{
+							EditMesh.AppendWithOffsets(AccumulatedMesh);
+						});
+					}
+					else
+					{
+						CurrentMesh->SetMesh(MoveTemp(AccumulatedMesh));
+						CurrentMesh = nullptr;
+					}
+				}
+			}
+		}
+
+		return true;
+	}
+} // namespace PCGGeometryHelpers
\ No newline at end of file
diff --git a/Engine/Plugins/PCGInterops/PCGGeometryScriptInterop/Source/PCGGeometryScriptInterop/Public/Helpers/PCGGeometryHelpers.h b/Engine/Plugins/PCGInterops/PCGGeometryScriptInterop/Source/PCGGeometryScriptInterop/Public/Helpers/PCGGeometryHelpers.h
--- a/Engine/Plugins/PCGInterops/PCGGeometryScriptInterop/Source/PCGGeometryScriptInterop/Public/Helpers/PCGGeometryHelpers.h
+++ b/Engine/Plugins/PCGInterops/PCGGeometryScriptInterop/Source/PCGGeometryScriptInterop/Public/Helpers/PCGGeometryHelpers.h
@@ -30,4 +30,6 @@ namespace PCGGeometryHelpers
 	 * @param OptionalMappings: Optional mappings if we want to update a subset of the triangles in the mesh.
 	 */
 	void RemapMaterials(UE::Geometry::FDynamicMesh3& InMesh, const TArray<UMaterialInterface*>& FromMaterials, TArray<TObjectPtr<UMaterialInterface>>& ToMaterials, const UE::Geometry::FMeshIndexMappings* OptionalMappings = nullptr);
+
+	bool ConvertDataToDynMeshes(const UPCGData* InData, FPCGContext* Context, TArray<UDynamicMesh*>& OutMeshes, bool bMergeMeshes, UGeometryScriptDebug* DynamicMeshDebug = nullptr);
 }
