[PCG] Remove automatic snapping of data viewport floor to the bounds of visualized data
* Caused confusion for users expecting their data to be transformed relative to the floor as if it were at the origin

#rnx

diff --git a/Engine/Plugins/PCG/Source/PCGEditor/Private/DataVisualizations/PCGSpatialDataVisualization.cpp b/Engine/Plugins/PCG/Source/PCGEditor/Private/DataVisualizations/PCGSpatialDataVisualization.cpp
--- a/Engine/Plugins/PCG/Source/PCGEditor/Private/DataVisualizations/PCGSpatialDataVisualization.cpp
+++ b/Engine/Plugins/PCG/Source/PCGEditor/Private/DataVisualizations/PCGSpatialDataVisualization.cpp
@@ -454,7 +454,6 @@ FPCGSetupSceneFunc IPCGSpatialDataVisualization::GetViewportSetupFunc(const UPCG
 		if (bInitializedBounds)
 		{
 			const FBoxSphereBounds Bounds = FBoxSphereBounds(FBox(BoundsMin, BoundsMax));
-			InOutParams.Scene->SetFloorOffset(-Bounds.Origin.Z + Bounds.BoxExtent.Z);
 
 			InOutParams.EditorViewportClient->SetViewLocation(Bounds.Origin);
 
diff --git a/Engine/Plugins/PCG/Source/PCGEditor/Private/DataVisualizations/PCGSplineDataVisualization.cpp b/Engine/Plugins/PCG/Source/PCGEditor/Private/DataVisualizations/PCGSplineDataVisualization.cpp
--- a/Engine/Plugins/PCG/Source/PCGEditor/Private/DataVisualizations/PCGSplineDataVisualization.cpp
+++ b/Engine/Plugins/PCG/Source/PCGEditor/Private/DataVisualizations/PCGSplineDataVisualization.cpp
@@ -167,7 +167,6 @@ FPCGSetupSceneFunc IPCGSplineDataVisualization::GetViewportSetupFunc(const UPCGD
 
 		FBoxSphereBounds Bounds = SplineComponent->CalcBounds(SplineComponent->GetComponentToWorld());
 
-		InOutParams.Scene->SetFloorOffset(-Bounds.Origin.Z + Bounds.BoxExtent.Z);
 		InOutParams.EditorViewportClient->SetViewLocation(Bounds.Origin);
 
 		if (Bounds.SphereRadius > 0.0f)
diff --git a/Engine/Plugins/PCG/Source/PCGEditor/Private/DataVisualizations/PCGStaticMeshDataVisualization.cpp b/Engine/Plugins/PCG/Source/PCGEditor/Private/DataVisualizations/PCGStaticMeshDataVisualization.cpp
--- a/Engine/Plugins/PCG/Source/PCGEditor/Private/DataVisualizations/PCGStaticMeshDataVisualization.cpp
+++ b/Engine/Plugins/PCG/Source/PCGEditor/Private/DataVisualizations/PCGStaticMeshDataVisualization.cpp
@@ -58,7 +58,6 @@ FPCGSetupSceneFunc IPCGStaticMeshDataVisualization::GetViewportSetupFunc(const U
 
 		// Bounds will be updated already by SetStaticMesh() call.
 		const FBoxSphereBounds Bounds = MeshComponent->Bounds;
-		InOutParams.Scene->SetFloorOffset(-Bounds.Origin.Z + Bounds.BoxExtent.Z);
 
 		InOutParams.EditorViewportClient->SetViewLocation(Bounds.Origin);
 
diff --git a/Engine/Plugins/PCG/Source/PCGEditor/Private/DataVisualizations/PCGVolumeDataVisualization.cpp b/Engine/Plugins/PCG/Source/PCGEditor/Private/DataVisualizations/PCGVolumeDataVisualization.cpp
--- a/Engine/Plugins/PCG/Source/PCGEditor/Private/DataVisualizations/PCGVolumeDataVisualization.cpp
+++ b/Engine/Plugins/PCG/Source/PCGEditor/Private/DataVisualizations/PCGVolumeDataVisualization.cpp
@@ -99,7 +99,6 @@ FPCGSetupSceneFunc FPCGPrimitiveDataVisualization::GetViewportSetupFunc(const UP
 
 		if (bHasBounds)
 		{
-			InOutParams.Scene->SetFloorOffset(-Bounds.Origin.Z + Bounds.BoxExtent.Z);
 			InOutParams.EditorViewportClient->SetViewLocation(Bounds.Origin);
 
 			if (Bounds.SphereRadius > 0.0f)
diff --git a/Engine/Plugins/PCGInterops/PCGGeometryScriptInterop/Source/PCGGeometryScriptInterop/Private/Editor/PCGDynamicMeshDataVisualization.cpp b/Engine/Plugins/PCGInterops/PCGGeometryScriptInterop/Source/PCGGeometryScriptInterop/Private/Editor/PCGDynamicMeshDataVisualization.cpp
--- a/Engine/Plugins/PCGInterops/PCGGeometryScriptInterop/Source/PCGGeometryScriptInterop/Private/Editor/PCGDynamicMeshDataVisualization.cpp
+++ b/Engine/Plugins/PCGInterops/PCGGeometryScriptInterop/Source/PCGGeometryScriptInterop/Private/Editor/PCGDynamicMeshDataVisualization.cpp
@@ -63,7 +63,6 @@ FPCGSetupSceneFunc FPCGDynamicMeshDataVisualization::GetViewportSetupFunc(const
 		InOutParams.Scene->AddComponent(DynamicMeshComponent, FTransform::Identity);
 
 		const FBoxSphereBounds Bounds = DynamicMeshComponent->CalcLocalBounds();
-		InOutParams.Scene->SetFloorOffset(-Bounds.Origin.Z + Bounds.BoxExtent.Z);
 
 		InOutParams.EditorViewportClient->SetViewLocation(Bounds.Origin);
 
