[PCG] Added null check to prevent crash when editing blueprint dealing with dynamic mesh data in PCG.

#rnx

diff --git a/Engine/Plugins/PCGInterops/PCGGeometryScriptInterop/Source/PCGGeometryScriptInterop/Private/Data/PCGDynamicMeshData.cpp b/Engine/Plugins/PCGInterops/PCGGeometryScriptInterop/Source/PCGGeometryScriptInterop/Private/Data/PCGDynamicMeshData.cpp
--- a/Engine/Plugins/PCGInterops/PCGGeometryScriptInterop/Source/PCGGeometryScriptInterop/Private/Data/PCGDynamicMeshData.cpp
+++ b/Engine/Plugins/PCGInterops/PCGGeometryScriptInterop/Source/PCGGeometryScriptInterop/Private/Data/PCGDynamicMeshData.cpp
@@ -24,6 +24,13 @@ UPCGDynamicMeshData::UPCGDynamicMeshData(const FObjectInitializer& ObjectInitial
 void UPCGDynamicMeshData::Initialize(UDynamicMesh* InMesh, bool bCanTakeOwnership, const TArray<UMaterialInterface*>& InMaterials)
 {
 	check(DynamicMesh);
+
+	if (!InMesh)
+	{
+		UE_LOG(LogPCG, Error, TEXT("Dynamic Mesh initialization failed because InMesh is null."));
+		return;
+	}
+
 	if (bCanTakeOwnership)
 	{
 		Initialize(std::move(InMesh->GetMeshRef()), InMaterials);
