[PCG] Improved tier height determination and fixed an invalid index in the tiers.

#rnx

diff --git a/Engine/Plugins/PCGInterops/PCGGeometryScriptInterop/Source/PCGGeometryScriptInterop/Private/Elements/PCGPrimitiveCrossSection.cpp b/Engine/Plugins/PCGInterops/PCGGeometryScriptInterop/Source/PCGGeometryScriptInterop/Private/Elements/PCGPrimitiveCrossSection.cpp
--- a/Engine/Plugins/PCGInterops/PCGGeometryScriptInterop/Source/PCGGeometryScriptInterop/Private/Elements/PCGPrimitiveCrossSection.cpp
+++ b/Engine/Plugins/PCGInterops/PCGGeometryScriptInterop/Source/PCGGeometryScriptInterop/Private/Elements/PCGPrimitiveCrossSection.cpp
@@ -155,56 +155,58 @@ bool FPCGPrimitiveCrossSectionElement::ExecuteInternal(FPCGContext* Context) con
 		}
 
 		FVector SliceOrigin = FVector::ZeroVector;
-		double MinProjectionScalar = std::numeric_limits<double>::max();
+		TArray<TPair<double, double>> TierHeights;
 
-		// Project vertices along the slice direction and find the minimum height
-		TArray<double> ProjectionScalars;
-		ProjectionScalars.Reserve(VertexList.List->Num());
-		for (const FVector& Vertex : *VertexList.List)
+		// Build tiers based on projection of vertices against the slicing axis
 		{
-			// Implementation note: since the direction vector is guaranteed to be normalized, we can use the projection scalar as distance
-			double ProjectionScalar = FVector::DotProduct(Vertex, SliceDirection);
-			ProjectionScalars.Emplace(ProjectionScalar);
+			double MinProjectionScalar = std::numeric_limits<double>::max();
 
-			// Find the min and max projected "height" to define the slicing extents
-			if (MinProjectionScalar > ProjectionScalar)
+			// Project vertices along the slice direction and find the minimum height
+			TArray<double> ProjectionScalars;
+			ProjectionScalars.Reserve(VertexList.List->Num());
+			for (const FVector& Vertex : *VertexList.List)
 			{
-				MinProjectionScalar = ProjectionScalar;
-				// Find the vertex for which we'll start slicing
-				SliceOrigin = Vertex;
-			}
-		}
+				// Implementation note: since the direction vector is guaranteed to be normalized, we can use the projection scalar as distance
+				double ProjectionScalar = FVector::DotProduct(Vertex, SliceDirection);
+				ProjectionScalars.Emplace(ProjectionScalar);
 
+				// Find the min and max projected "height" to define the slicing extents
+				if (MinProjectionScalar > ProjectionScalar)
+				{
+					MinProjectionScalar = ProjectionScalar;
+					// Find the vertex for which we'll start slicing
+					SliceOrigin = Vertex;
+				}
+			}
 
-		TArray<double> TierHeights;
-		// Merge at the user's threshold, or at a minimum threshold to account for numerical or mesh imprecision.
-		const double MergingThreshold = Settings->bEnableTierMerging ? Settings->TierMergingThreshold : PCGPrimitiveCrossSection::Constants::MinTierMergingThreshold;
-		TierHeights.Reserve(ProjectionScalars.Num() + 1);
+			// Merge at the user's threshold, or at a minimum threshold to account for numerical or mesh imprecision.
+			const double MergingThreshold = Settings->bEnableTierMerging ? Settings->TierMergingThreshold : PCGPrimitiveCrossSection::Constants::MinTierMergingThreshold;
+			TierHeights.Reserve(ProjectionScalars.Num() + 1);
 
-		// Once sorted, all similar tier ranges will be consecutive.
-		ProjectionScalars.Sort();
+			// Once sorted, all similar tier ranges will be consecutive.
+			ProjectionScalars.Sort();
 
-		int32 LastGroupIndex = 0;
-		int32 CurrentIndex = 1;
-		while (CurrentIndex < ProjectionScalars.Num())
-		{
-			// Consecutive value falls within the threshold. Merge them.
-			if (ProjectionScalars[CurrentIndex] - ProjectionScalars[LastGroupIndex] >= MergingThreshold)
+			for (const double& Projection : ProjectionScalars)
 			{
-				TierHeights.Emplace(ProjectionScalars[LastGroupIndex]);
-				LastGroupIndex = CurrentIndex;
-			}
-
-			++CurrentIndex;
-		}
+				double AdjustedProjection = Projection - MinProjectionScalar;
+				bool bNeedsNewTier = false;
 
-		// Add the final range.
-		TierHeights.Emplace(ProjectionScalars[LastGroupIndex]);
+				// Update previous tier if there is one
+				if (!TierHeights.IsEmpty())
+				{
+					bNeedsNewTier = (TierHeights.Last().Value - AdjustedProjection) >= MergingThreshold;
+					TierHeights.Last().Value = AdjustedProjection;
+				}
+				else
+				{
+					bNeedsNewTier = true;
+				}
 
-		// Convert projection values to height values.
-		for (double& Height : TierHeights)
-		{
-			Height -= MinProjectionScalar;
+				if (bNeedsNewTier)
+				{
+					TierHeights.Emplace(AdjustedProjection, AdjustedProjection);
+				}
+			}
 		}
 
 		// Set up the slice plane for GeometryScript, represented by an FTransform
@@ -216,7 +218,7 @@ bool FPCGPrimitiveCrossSectionElement::ExecuteInternal(FPCGContext* Context) con
 		TArray<FCrossSection> CrossSections;
 		for (int TierIndex = 0; TierIndex < TierHeights.Num(); ++TierIndex)
 		{
-			FVector SliceLocation = SliceOrigin + SliceDirection * TierHeights[TierIndex];
+			FVector SliceLocation = SliceOrigin + SliceDirection * TierHeights[TierIndex].Key;
 			SlicePlaneTransform.SetLocation(SliceLocation);
 
 			// Cuts the mesh at the specified plane, leaving a hole
@@ -328,9 +330,16 @@ bool FPCGPrimitiveCrossSectionElement::ExecuteInternal(FPCGContext* Context) con
 					continue;
 				}
 
+				const double Height = (TierHeights[TierIndex].Value - TierHeights[TierIndex].Key);
+
+				if (Settings->bEnableMinHeightCulling && Height <= Settings->MinHeightCullingThreshold)
+				{
+					continue;
+				}
+
 				FCrossSection& CurrentTier = CrossSections.Emplace_GetRef();
 				CurrentTier.Tier = TierIndex;
-				const double Height = TierHeights[TierIndex + 1] - TierHeights[TierIndex];
+				
 				const double RoundedHeight = FMath::RoundToDouble(Height);
 				// To account for rounding errors and mesh operations.
 				CurrentTier.Height = FMath::IsNearlyEqual(Height, RoundedHeight) ? RoundedHeight : Height;
diff --git a/Engine/Plugins/PCGInterops/PCGGeometryScriptInterop/Source/PCGGeometryScriptInterop/Public/Elements/PCGPrimitiveCrossSection.h b/Engine/Plugins/PCGInterops/PCGGeometryScriptInterop/Source/PCGGeometryScriptInterop/Public/Elements/PCGPrimitiveCrossSection.h
--- a/Engine/Plugins/PCGInterops/PCGGeometryScriptInterop/Source/PCGGeometryScriptInterop/Public/Elements/PCGPrimitiveCrossSection.h
+++ b/Engine/Plugins/PCGInterops/PCGGeometryScriptInterop/Source/PCGGeometryScriptInterop/Public/Elements/PCGPrimitiveCrossSection.h
@@ -69,6 +69,14 @@ public:
 	UPROPERTY(BlueprintReadWrite, EditAnywhere, Category = Settings, meta = (PCG_Overridable, EditCondition = "bEnableMinAreaCulling"))
 	double MinAreaCullingThreshold = 100.0;
 
+	/** Culls tiers that don't meed a minimum height requirement. */
+	UPROPERTY(BlueprintReadWrite, EditAnywhere, Category = Settings, meta = (PCG_Overridable, InlineEditConditionToggle))
+	bool bEnableMinHeightCulling = true;
+
+	/** If a tier is smaller in height than this threshold, it will be culled. */
+	UPROPERTY(BlueprintReadWrite, EditAnywhere, Category = Settings, meta = (PCG_Overridable, EditCondition = "bEnableMinHeightCulling"))
+	double MinHeightCullingThreshold = 1.0;
+
 	/** If multiple tiers can be combined into a single tier without affecting the contour, remove the redundant one. Note: This will currently cull even if there are other unique tiers in between. */
 	UPROPERTY(BlueprintReadWrite, EditAnywhere, Category = Settings, meta = (PCG_Overridable))
 	bool bRemoveRedundantSections = true;
