[PCG] Simplified output mesh from collision shape wrapper to make geometry operations more stable

#rnx

diff --git a/Engine/Plugins/PCGInterops/PCGGeometryScriptInterop/Source/PCGGeometryScriptInterop/Private/Helpers/PCGGeometryHelpers.cpp b/Engine/Plugins/PCGInterops/PCGGeometryScriptInterop/Source/PCGGeometryScriptInterop/Private/Helpers/PCGGeometryHelpers.cpp
--- a/Engine/Plugins/PCGInterops/PCGGeometryScriptInterop/Source/PCGGeometryScriptInterop/Private/Helpers/PCGGeometryHelpers.cpp
+++ b/Engine/Plugins/PCGInterops/PCGGeometryScriptInterop/Source/PCGGeometryScriptInterop/Private/Helpers/PCGGeometryHelpers.cpp
@@ -313,7 +313,7 @@ namespace PCGGeometryHelpers
 					}
 					else
 					{
-						UE::Geometry::ConvertSimpleCollisionToMeshes(BodySetup->AggGeom, AccumulatedMesh, Transforms, SphereResolution, true, true, nullptr, false, ExternalScale);
+						UE::Geometry::ConvertSimpleCollisionToMeshes(BodySetup->AggGeom, AccumulatedMesh, Transforms, SphereResolution, false, false, nullptr, false, ExternalScale);
 						bHasData = true;
 					}
 				}
