[PCG] Allow dynamic mesh data to be used in primitive cross section node

#rnx

diff --git a/Engine/Plugins/PCGInterops/PCGGeometryScriptInterop/Source/PCGGeometryScriptInterop/Private/Elements/PCGPrimitiveCrossSection.cpp b/Engine/Plugins/PCGInterops/PCGGeometryScriptInterop/Source/PCGGeometryScriptInterop/Private/Elements/PCGPrimitiveCrossSection.cpp
--- a/Engine/Plugins/PCGInterops/PCGGeometryScriptInterop/Source/PCGGeometryScriptInterop/Private/Elements/PCGPrimitiveCrossSection.cpp
+++ b/Engine/Plugins/PCGInterops/PCGGeometryScriptInterop/Source/PCGGeometryScriptInterop/Private/Elements/PCGPrimitiveCrossSection.cpp
@@ -5,6 +5,7 @@
 
 #include "PCGComponent.h"
 #include "PCGContext.h"
+#include "Data/PCGDynamicMeshData.h"
 #include "Data/PCGPrimitiveData.h"
 #include "Data/PCGSplineData.h"
 #include "Helpers/PCGGeometryHelpers.h"
@@ -47,7 +48,7 @@ UPCGPrimitiveCrossSectionSettings::UPCGPrimitiveCrossSectionSettings()
 TArray<FPCGPinProperties> UPCGPrimitiveCrossSectionSettings::InputPinProperties() const
 {
 	TArray<FPCGPinProperties> PinProperties;
-	PinProperties.Emplace(PCGPinConstants::DefaultInputLabel, EPCGDataType::Primitive);
+	PinProperties.Emplace(PCGPinConstants::DefaultInputLabel, EPCGDataType::Volume | EPCGDataType::Primitive | EPCGDataType::DynamicMesh);
 	return PinProperties;
 }
 
diff --git a/Engine/Plugins/PCGInterops/PCGGeometryScriptInterop/Source/PCGGeometryScriptInterop/Private/Helpers/PCGGeometryHelpers.cpp b/Engine/Plugins/PCGInterops/PCGGeometryScriptInterop/Source/PCGGeometryScriptInterop/Private/Helpers/PCGGeometryHelpers.cpp
--- a/Engine/Plugins/PCGInterops/PCGGeometryScriptInterop/Source/PCGGeometryScriptInterop/Private/Helpers/PCGGeometryHelpers.cpp
+++ b/Engine/Plugins/PCGInterops/PCGGeometryScriptInterop/Source/PCGGeometryScriptInterop/Private/Helpers/PCGGeometryHelpers.cpp
@@ -5,6 +5,7 @@
 #include "ConversionUtils/SceneComponentToDynamicMesh.h"
 #include "Data/PCGCollisionShapeData.h"
 #include "Data/PCGCollisionWrapperData.h"
+#include "Data/PCGDynamicMeshData.h"
 #include "Data/PCGPrimitiveData.h"
 #include "Data/PCGVolumeData.h"
 #include "DynamicMesh/DynamicMeshAttributeSet.h"
@@ -348,6 +349,14 @@ namespace PCGGeometryHelpers
 				}
 			}
 		}
+		else if (const UPCGDynamicMeshData* DynamicMeshData = Cast<UPCGDynamicMeshData>(InData))
+		{
+			if (DynamicMeshData->GetDynamicMesh())
+			{
+				UDynamicMesh* NewMesh = OutMeshes.Emplace_GetRef(FPCGContext::NewObject_AnyThread<UDynamicMesh>(Context));
+				NewMesh->SetMesh(DynamicMeshData->GetDynamicMesh()->GetMeshRef());
+			}
+		}
 
 		return true;
 	}
