[PCG][Bughawk] Early out when ProcessDynamicMesh not implemented and updated comment

#rnx

diff --git a/Engine/Plugins/PCGInterops/PCGGeometryScriptInterop/Source/PCGGeometryScriptInterop/Private/Elements/PCGGeometryBlueprintElement.cpp b/Engine/Plugins/PCGInterops/PCGGeometryScriptInterop/Source/PCGGeometryScriptInterop/Private/Elements/PCGGeometryBlueprintElement.cpp
--- a/Engine/Plugins/PCGInterops/PCGGeometryScriptInterop/Source/PCGGeometryScriptInterop/Private/Elements/PCGGeometryBlueprintElement.cpp
+++ b/Engine/Plugins/PCGInterops/PCGGeometryScriptInterop/Source/PCGGeometryScriptInterop/Private/Elements/PCGGeometryBlueprintElement.cpp
@@ -31,12 +31,15 @@ void UPCGGeometryBlueprintElement::Execute_Implementation(const FPCGDataCollecti
 	const bool bHasASingleOutputPin = (bHasDefaultOutPin && CustomOutputPins.IsEmpty()) || (!bHasDefaultOutPin && CustomOutputPins.Num() == 1);
 
 	const bool bIsProcessDynMeshOverridden =  GetClass()->IsFunctionImplementedInScript(GET_FUNCTION_NAME_CHECKED(UPCGGeometryBlueprintElement, ProcessDynamicMesh));
-	
-	if ((!bHasASingleInputPin || !bHasASingleOutputPin) && bIsProcessDynMeshOverridden)
+	if (!bIsProcessDynMeshOverridden)
+	{
+		return;
+	}
+
+	if (!bHasASingleInputPin || !bHasASingleOutputPin)
 	{
 		// Make sure to throw a warning if we are in this case and Process Dynamic Mesh is overridden to warn the user their function won't be called
-		PCGLog::LogWarningOnGraph(LOCTEXT("DynMeshOverridenButNotCalled", "Process Dynamic Mesh was overridden, but we don't have the expected setup (single input and output pin)"
-				" or Execute is also overriden. Process Dynamic Mesh won't be called."), &GetContext());
+		PCGLog::LogWarningOnGraph(LOCTEXT("DynMeshOverridenButNotCalled", "Process Dynamic Mesh was overridden, but we don't have the expected setup (single input and output pin). Process Dynamic Mesh won't be called."), &GetContext());
 		return;
 	}
 
