[PCG] Fix RemapMaterial not working if the target mesh didn't have any material

* This was breaking appending a mesh with materials to an empty mesh for example

#rnx

diff --git a/Engine/Plugins/PCGInterops/PCGGeometryScriptInterop/Source/PCGGeometryScriptInterop/Private/Helpers/PCGGeometryHelpers.cpp b/Engine/Plugins/PCGInterops/PCGGeometryScriptInterop/Source/PCGGeometryScriptInterop/Private/Helpers/PCGGeometryHelpers.cpp
--- a/Engine/Plugins/PCGInterops/PCGGeometryScriptInterop/Source/PCGGeometryScriptInterop/Private/Helpers/PCGGeometryHelpers.cpp
+++ b/Engine/Plugins/PCGInterops/PCGGeometryScriptInterop/Source/PCGGeometryScriptInterop/Private/Helpers/PCGGeometryHelpers.cpp
@@ -80,11 +80,18 @@ namespace PCGGeometryHelpers
 
 	void RemapMaterials(UE::Geometry::FDynamicMesh3& InMesh, const TArray<UMaterialInterface*>& FromMaterials, TArray<TObjectPtr<UMaterialInterface>>& ToMaterials, const UE::Geometry::FMeshIndexMappings* OptionalMappings)
 	{
-		if (FromMaterials.IsEmpty() || ToMaterials.IsEmpty() || !InMesh.HasAttributes() || !InMesh.Attributes()->HasMaterialID())
+		if (FromMaterials.IsEmpty() || !InMesh.HasAttributes() || !InMesh.Attributes()->HasMaterialID())
 		{
 			return;
 		}
-		
+
+		// If we have no materials in output, we can just copy the in materials
+		if (ToMaterials.IsEmpty())
+		{
+			ToMaterials = FromMaterials;
+			return;
+		}
+
 		TMap<int32, int32> MaterialIDRemap;
 		MaterialIDRemap.Reserve(FromMaterials.Num());
 		
