[PCG] Fix Cross Section issue: Projection heights were not considered for cross sections as comparison with previous height was reversed.

#rnx

diff --git a/Engine/Plugins/PCGInterops/PCGGeometryScriptInterop/Source/PCGGeometryScriptInterop/Private/Elements/PCGPrimitiveCrossSection.cpp b/Engine/Plugins/PCGInterops/PCGGeometryScriptInterop/Source/PCGGeometryScriptInterop/Private/Elements/PCGPrimitiveCrossSection.cpp
--- a/Engine/Plugins/PCGInterops/PCGGeometryScriptInterop/Source/PCGGeometryScriptInterop/Private/Elements/PCGPrimitiveCrossSection.cpp
+++ b/Engine/Plugins/PCGInterops/PCGGeometryScriptInterop/Source/PCGGeometryScriptInterop/Private/Elements/PCGPrimitiveCrossSection.cpp
@@ -197,7 +197,7 @@ bool FPCGPrimitiveCrossSectionElement::ExecuteInternal(FPCGContext* Context) con
 				// Update previous tier if there is one
 				if (!TierHeights.IsEmpty())
 				{
-					bNeedsNewTier = (TierHeights.Last().Value - AdjustedProjection) >= MergingThreshold;
+					bNeedsNewTier = (AdjustedProjection - TierHeights.Last().Value) >= MergingThreshold;
 					TierHeights.Last().Value = AdjustedProjection;
 				}
 				else
