[PCG] Enhance the Mesh Sampler to support classes and blueprint assets as input

* Allow to pass an asset of a Blueprint Class that has a Static mesh component for example without spawning the actor
* Will look for the default actor of the class and the default BP component of the class if it is a BP.

diff --git a/Engine/Plugins/PCGInterops/PCGGeometryScriptInterop/Source/PCGGeometryScriptInterop/Private/Elements/PCGMeshSampler.cpp b/Engine/Plugins/PCGInterops/PCGGeometryScriptInterop/Source/PCGGeometryScriptInterop/Private/Elements/PCGMeshSampler.cpp
--- a/Engine/Plugins/PCGInterops/PCGGeometryScriptInterop/Source/PCGGeometryScriptInterop/Private/Elements/PCGMeshSampler.cpp
+++ b/Engine/Plugins/PCGInterops/PCGGeometryScriptInterop/Source/PCGGeometryScriptInterop/Private/Elements/PCGMeshSampler.cpp
@@ -18,6 +18,7 @@
 #include "ConversionUtils/SceneComponentToDynamicMesh.h"
 #include "Data/PCGDynamicMeshData.h"
 #include "Engine/AssetManager.h"
+#include "Engine/Blueprint.h"
 #include "Engine/StaticMesh.h"
 #include "GameFramework/Actor.h"
 #include "GeometryScript/MeshAssetFunctions.h"
@@ -496,73 +497,114 @@ bool FPCGMeshSamplerElement::PrepareDataInternal(FPCGContext* InContext) const
 			return true;
 		}
 
+		// For tracking purposes, keep the original object
+		const UObject* OriginalObject = Object;
+
+		// Check if the object is a Blueprint of type Actor, if so, extract the class to be read just after
+		if (const UBlueprint* Blueprint = Cast<const UBlueprint>(Object))
+		{
+			if (UClass* BlueprintClass = Blueprint->GeneratedClass.Get(); BlueprintClass && BlueprintClass->IsChildOf<AActor>())
+			{
+				Object = BlueprintClass;
+			}
+		}
+
+		TArray<UObject*> Candidates;
+
 		EGeometryScriptOutcomePins Outcome = EGeometryScriptOutcomePins::Success;
 		Debug->Messages.Reset();
-		USceneComponent* SceneComponent = Cast<USceneComponent>(Object);
 
-		// @todo_pcg: Better support of multiple scene component on a given actor.
-		if (AActor* Actor = Cast<AActor>(Object))
+		// Check if the UObject is a class of an actor. If so, find the all scene components
+		if (UClass* Class = Cast<UClass>(Object); Class && Class->IsChildOf<AActor>())
 		{
-			SceneComponent = Actor->GetRootComponent();
+			AActor::ForEachComponentOfActorClassDefault<USceneComponent>(Class, [&Candidates](const USceneComponent* Component)
+				{ 
+					Candidates.Add(const_cast<USceneComponent*>(Component));
+					return true; // Continue until we find them all
+				});
 		}
-
-		if (SceneComponent)
+		else if (AActor* Actor = Cast<AActor>(Object))
 		{
-#if WITH_EDITOR
-			DynamicTracking.AddToTracking(FPCGSelectionKey::CreateFromPath(FSoftObjectPath(SceneComponent)), /*bIsCulled=*/false);
-#endif // WITH_EDITOR
+			Actor->ForEachComponent<USceneComponent>(/*bIncludeFromChildActor=*/ false, [&Candidates](const USceneComponent* Component)
+				{
+					Candidates.Add(const_cast<USceneComponent*>(Component));
+				});
+		}
+		else
+		{
+			Candidates.Add(Object);
+		}
 
-			// Adaptation of UGeometryScriptLibrary_SceneUtilityFunctions::CopyMeshFromComponent, since we don't have access to the Material list with this one.
-			Context->DynamicMeshes.Add(FPCGContext::NewObject_AnyThread<UDynamicMesh>(Context));
-			FTransform Transform;
-			FText ErrorMessage;
-			UE::Conversion::FToMeshOptions Options{};
-			Options.bWantInstanceColors = true;
-			Options.LODType = PCGGeometryHelpers::SafeConversionLODType(Settings->RequestedLODType);
-			Options.LODIndex = Settings->RequestedLODIndex;
+		bool bValid = false;
 
-			UE::Geometry::FDynamicMesh3 TempDynMesh{};
+		// @todo_pcg: Better support of multiple candidates. For now, take the first valid one.
+		for (UObject* Candidate : Candidates)
+		{
+			if (USceneComponent* SceneComponent = Cast<USceneComponent>(Candidate))
+			{
+				// Adaptation of UGeometryScriptLibrary_SceneUtilityFunctions::CopyMeshFromComponent, since we don't have access to the Material list with this one.
+				Context->DynamicMeshes.Add(FPCGContext::NewObject_AnyThread<UDynamicMesh>(Context));
+				FTransform Transform;
+				FText ErrorMessage;
+				UE::Conversion::FToMeshOptions Options{};
+				Options.bWantInstanceColors = true;
+				Options.LODType = PCGGeometryHelpers::SafeConversionLODType(Settings->RequestedLODType);
+				Options.LODIndex = Settings->RequestedLODIndex;
+
+				UE::Geometry::FDynamicMesh3 TempDynMesh{};
+
+				TArray<UMaterialInterface*>* ComponentMaterialListPtr = Settings->bOutputMaterialInfo ? &Context->ComponentMaterialList.Emplace_GetRef() : nullptr;
+				TArray<UMaterialInterface*>* AssetMaterialListPtr = Settings->bOutputMaterialInfo ? &Context->AssetMaterialList.Emplace_GetRef() : nullptr;
+
+				const bool bSuccess = UE::Conversion::SceneComponentToDynamicMesh(SceneComponent, Options, /*bTransformToWorld=*/false, TempDynMesh, Transform, ErrorMessage, ComponentMaterialListPtr, AssetMaterialListPtr);
+				if (!bSuccess)
+				{
+					Outcome = EGeometryScriptOutcomePins::Failure;
+					FGeometryScriptDebugMessage DebugMessage;
+					DebugMessage.Message = std::move(ErrorMessage);
+					Debug->Messages.Emplace(std::move(DebugMessage));
+				}
+				else
+				{
+					Outcome = EGeometryScriptOutcomePins::Success;
+					bValid = true;
+					Context->DynamicMeshes.Last()->SetMesh(std::move(TempDynMesh));
 
-			TArray<UMaterialInterface*>* ComponentMaterialListPtr = Settings->bOutputMaterialInfo ? &Context->ComponentMaterialList.Emplace_GetRef() : nullptr;
-			TArray<UMaterialInterface*>* AssetMaterialListPtr = Settings->bOutputMaterialInfo ? &Context->AssetMaterialList.Emplace_GetRef() : nullptr;
+#if WITH_EDITOR
+					// Track the scene component and also the original object, only if the extraction worked
+					DynamicTracking.AddToTracking(FPCGSelectionKey::CreateFromPath(FSoftObjectPath(OriginalObject)), /*bIsCulled=*/false);
+					DynamicTracking.AddToTracking(FPCGSelectionKey::CreateFromPath(FSoftObjectPath(SceneComponent)), /*bIsCulled=*/false);
+#endif // WITH_EDITOR
 
-			const bool bSuccess = UE::Conversion::SceneComponentToDynamicMesh(SceneComponent, Options, /*bTransformToWorld=*/false, TempDynMesh, Transform, ErrorMessage, ComponentMaterialListPtr, AssetMaterialListPtr);
-			if (!bSuccess)
-			{
-				Outcome = EGeometryScriptOutcomePins::Failure;
-				FGeometryScriptDebugMessage DebugMessage;
-				DebugMessage.Message = std::move(ErrorMessage);
-				Debug->Messages.Emplace(std::move(DebugMessage));
+					break;
+				}
 			}
-			else
+			else if (UStaticMesh* StaticMesh = Cast<UStaticMesh>(Candidate))
 			{
-				Context->DynamicMeshes.Last()->SetMesh(std::move(TempDynMesh));
-			}
-		}
-		else if (UStaticMesh* StaticMesh = Cast<UStaticMesh>(Object))
-		{
+				const FGeometryScriptMeshReadLOD MeshReadLOD{ Settings->RequestedLODType, Settings->RequestedLODIndex };
 
-#if WITH_EDITOR
-			DynamicTracking.AddToTracking(FPCGSelectionKey::CreateFromPath(FSoftObjectPath(StaticMesh)), /*bIsCulled=*/false);
-#endif // WITH_EDITOR
+				Context->DynamicMeshes.Add(FPCGContext::NewObject_AnyThread<UDynamicMesh>(Context));
+				FGeometryScriptCopyMeshFromAssetOptions Options{};
 
-			const FGeometryScriptMeshReadLOD MeshReadLOD{ Settings->RequestedLODType, Settings->RequestedLODIndex };
+				UGeometryScriptLibrary_StaticMeshFunctions::CopyMeshFromStaticMeshV2(StaticMesh, Context->DynamicMeshes.Last(), Options, MeshReadLOD, Outcome, false, Debug);
+				if (Outcome == EGeometryScriptOutcomePins::Success)
+				{
+					if (Settings->bOutputMaterialInfo)
+					{
+						TArray<FName> MaterialSlotNames;
+						UGeometryScriptLibrary_StaticMeshFunctions::GetMaterialListFromStaticMesh(StaticMesh, Context->AssetMaterialList.Emplace_GetRef(), MaterialSlotNames);
+					}
 
-			Context->DynamicMeshes.Add(FPCGContext::NewObject_AnyThread<UDynamicMesh>(Context));
-			FGeometryScriptCopyMeshFromAssetOptions Options{};
+#if WITH_EDITOR
+					// Track the static mesh, only if the extraction worked
+					DynamicTracking.AddToTracking(FPCGSelectionKey::CreateFromPath(FSoftObjectPath(StaticMesh)), /*bIsCulled=*/false);
+#endif // WITH_EDITOR
 
-			UGeometryScriptLibrary_StaticMeshFunctions::CopyMeshFromStaticMeshV2(StaticMesh, Context->DynamicMeshes.Last(), Options, MeshReadLOD, Outcome, false, Debug);
-			if (Outcome == EGeometryScriptOutcomePins::Success && Settings->bOutputMaterialInfo)
-			{
-				TArray<FName> MaterialSlotNames;
-				UGeometryScriptLibrary_StaticMeshFunctions::GetMaterialListFromStaticMesh(StaticMesh, Context->AssetMaterialList.Emplace_GetRef(), MaterialSlotNames);
+					bValid = true;
+					break;
+				}
 			}
 		}
-		else
-		{
-			PCGE_LOG(Error, GraphAndLog, FText::Format(LOCTEXT("ObjectNotRightType", "Provided object '{0}' is not a supported type. Only supports StaticMesh/Actor/SceneComponent."), FText::FromString(Path.ToString())));
-			continue;
-		}
 
 		if (Outcome != EGeometryScriptOutcomePins::Success)
 		{
@@ -571,6 +613,11 @@ bool FPCGMeshSamplerElement::PrepareDataInternal(FPCGContext* InContext) const
 			Context->DynamicMeshes.Last()->MarkAsGarbage();
 			Context->DynamicMeshes.RemoveAtSwap(Context->DynamicMeshes.Num() - 1);
 		}
+		else if (!bValid)
+		{
+			PCGE_LOG(Error, GraphAndLog, FText::Format(LOCTEXT("ObjectNotRightType", "Provided object '{0}' is not a supported type. Only supports StaticMesh/Actor/SceneComponent."), FText::FromString(Path.ToString())));
+			continue;
+		}
 	}
 
 #if WITH_EDITOR
