[PCG] Added an optional vector mode for extrusion, and made thickness overridable in the spline to mesh node.

diff --git a/Engine/Plugins/PCGInterops/PCGGeometryScriptInterop/Source/PCGGeometryScriptInterop/PCGGeometryScriptInterop.Build.cs b/Engine/Plugins/PCGInterops/PCGGeometryScriptInterop/Source/PCGGeometryScriptInterop/PCGGeometryScriptInterop.Build.cs
--- a/Engine/Plugins/PCGInterops/PCGGeometryScriptInterop/Source/PCGGeometryScriptInterop/PCGGeometryScriptInterop.Build.cs
+++ b/Engine/Plugins/PCGInterops/PCGGeometryScriptInterop/Source/PCGGeometryScriptInterop/PCGGeometryScriptInterop.Build.cs
@@ -24,6 +24,8 @@ namespace UnrealBuildTool.Rules
 			PrivateDependencyModuleNames.AddRange(
 				new string[]
 				{
+					"DynamicMesh",
+					"GeometryAlgorithms",
 					"GeometryCore",
 					"GeometryFramework",
 					"ModelingComponents"
diff --git a/Engine/Plugins/PCGInterops/PCGGeometryScriptInterop/Source/PCGGeometryScriptInterop/Private/Elements/PCGSplineToMesh.cpp b/Engine/Plugins/PCGInterops/PCGGeometryScriptInterop/Source/PCGGeometryScriptInterop/Private/Elements/PCGSplineToMesh.cpp
--- a/Engine/Plugins/PCGInterops/PCGGeometryScriptInterop/Source/PCGGeometryScriptInterop/Private/Elements/PCGSplineToMesh.cpp
+++ b/Engine/Plugins/PCGInterops/PCGGeometryScriptInterop/Source/PCGGeometryScriptInterop/Private/Elements/PCGSplineToMesh.cpp
@@ -7,6 +7,7 @@
 #include "Data/PCGSplineData.h"
 
 #include "CurveOps/TriangulateCurvesOp.h"
+#include "Operations/ExtrudeMesh.h"
 
 #include UE_INLINE_GENERATED_CPP_BY_NAME(PCGSplineToMesh)
 
@@ -25,7 +26,7 @@ FText UPCGSplineToMeshSettings::GetDefaultNodeTitle() const
 
 FText UPCGSplineToMeshSettings::GetNodeTooltipText() const
 {
-	return LOCTEXT("NodeTooltip", "Converts a closed spline into a mesh.");
+	return LOCTEXT("NodeTooltip", "Converts a spline into a closed mesh, and optionally extrudes it according to a direction. Note that the operation will 'close' all splines.");
 }
 #endif // WITH_EDITOR
 
@@ -50,6 +51,8 @@ bool FPCGSplineToMeshElement::ExecuteInternal(FPCGContext* InContext) const
 	const UPCGSplineToMeshSettings* Settings = InContext->GetInputSettings<UPCGSplineToMeshSettings>();
 	check(Settings);
 
+	const bool bUseDefaultExtrusion = (Settings->ExtrusionVector - FVector::ZAxisVector).IsNearlyZero(UE_SMALL_NUMBER) && Settings->bExtrudeInLocalSpace;
+
 	const TArray<FPCGTaggedData> Inputs = InContext->InputData.GetInputsByPin(PCGPinConstants::DefaultInputLabel);
 	for (const FPCGTaggedData& Input : Inputs)
 	{
@@ -61,7 +64,7 @@ bool FPCGSplineToMeshElement::ExecuteInternal(FPCGContext* InContext) const
 		}
 
 		UE::Geometry::FTriangulateCurvesOp TriangulateCurvesOp;
-		TriangulateCurvesOp.Thickness = Settings->Thickness;
+		TriangulateCurvesOp.Thickness = bUseDefaultExtrusion ? Settings->Thickness : 0.0;
 		TriangulateCurvesOp.bFlipResult = Settings->bFlipResult;
 		TriangulateCurvesOp.FlattenMethod = Settings->FlattenMethod;
 		TriangulateCurvesOp.CurveOffset = Settings->CurveOffset;
@@ -92,10 +95,30 @@ bool FPCGSplineToMeshElement::ExecuteInternal(FPCGContext* InContext) const
 			PCGLog::LogErrorOnGraph(LOCTEXT("TriangulationFailed", "Triangulation failed"), InContext);
 			continue;
 		}
+
+		if (!bUseDefaultExtrusion)
+		{
+			UE::Geometry::FExtrudeMesh ExtrudeMesh(DynamicMesh.Get());
+			ExtrudeMesh.ExtrudedPositionFunc = [Settings](const FVector3d& P, const FVector3f& N, int VID) -> FVector3d 
+			{ 
+				if (Settings->bExtrudeInLocalSpace)
+				{
+					FQuat RelativeRotation = FQuat::FindBetweenNormals(FVector::ZAxisVector, FVector(N));
+					return P + RelativeRotation.RotateVector(Settings->ExtrusionVector) * Settings->Thickness;
+				}
+				else
+				{
+					return P + Settings->ExtrusionVector * Settings->Thickness;
+				}
+			};
+			ExtrudeMesh.Apply();
+		}
 		
 		UE::Geometry::FDynamicMesh3* DynamicMeshPtr = DynamicMesh.Release();
 		UPCGDynamicMeshData* DynamicMeshData = FPCGContext::NewObject_AnyThread<UPCGDynamicMeshData>(InContext);
 		DynamicMeshData->Initialize(std::move(*DynamicMeshPtr));
+		delete DynamicMeshPtr;
+		DynamicMeshPtr = nullptr;
 
 		InContext->OutputData.TaggedData.Emplace_GetRef(Input).Data = DynamicMeshData;
 	}
diff --git a/Engine/Plugins/PCGInterops/PCGGeometryScriptInterop/Source/PCGGeometryScriptInterop/Public/Elements/PCGSplineToMesh.h b/Engine/Plugins/PCGInterops/PCGGeometryScriptInterop/Source/PCGGeometryScriptInterop/Public/Elements/PCGSplineToMesh.h
--- a/Engine/Plugins/PCGInterops/PCGGeometryScriptInterop/Source/PCGGeometryScriptInterop/Public/Elements/PCGSplineToMesh.h
+++ b/Engine/Plugins/PCGInterops/PCGGeometryScriptInterop/Source/PCGGeometryScriptInterop/Public/Elements/PCGSplineToMesh.h
@@ -38,9 +38,15 @@ public:
 	EFlattenCurveMethod FlattenMethod = EFlattenCurveMethod::DoNotFlatten;
 
 	/** If > 0, Extrude the triangulation by this amount. */
-	UPROPERTY(EditAnywhere, Category = Mesh, meta = (ClampMin = 0.0))
+	UPROPERTY(EditAnywhere, Category = Mesh, meta = (PCG_Overridable, ClampMin = 0.0))
 	double Thickness = 0.0;
 
+	UPROPERTY(EditAnywhere, Category = Mesh, meta = (PCG_Overridable))
+	FVector ExtrusionVector = FVector::ZAxisVector;
+
+	UPROPERTY(EditAnywhere, Category = Mesh)
+	bool bExtrudeInLocalSpace = true;
+
 	/** Whether to flip the facing direction of the generated mesh. */
 	UPROPERTY(EditAnywhere, Category = Mesh)
 	bool bFlipResult = false;
