[PCG] Add Get/Set Materials on Dynamic mesh nodes

#rnx
#Synced-CL 48550537

diff --git a/Engine/Plugins/PCG/Source/PCG/Private/Data/PCGDynamicMeshData.cpp b/Engine/Plugins/PCG/Source/PCG/Private/Data/PCGDynamicMeshData.cpp
--- a/Engine/Plugins/PCG/Source/PCG/Private/Data/PCGDynamicMeshData.cpp
+++ b/Engine/Plugins/PCG/Source/PCG/Private/Data/PCGDynamicMeshData.cpp
@@ -227,8 +227,12 @@ const UPCGPointData* UPCGDynamicMeshData::ToPointData(FPCGContext* Context, cons
 
 void UPCGDynamicMeshData::SetMaterials(const TArray<UMaterialInterface*>& InMaterials)
 {
-	Materials.Reset(InMaterials.Num());
-	Materials.Append(InMaterials);
+	Materials = InMaterials;
+}
+
+void UPCGDynamicMeshData::SetMaterials(TArray<UMaterialInterface*>&& InMaterials)
+{
+	Materials = MoveTemp(InMaterials);
 }
 
 void UPCGDynamicMeshData::InitializeDynamicMeshComponentFromData(UDynamicMeshComponent* InComponent) const
diff --git a/Engine/Plugins/PCG/Source/PCG/Public/Data/PCGDynamicMeshData.h b/Engine/Plugins/PCG/Source/PCG/Public/Data/PCGDynamicMeshData.h
--- a/Engine/Plugins/PCG/Source/PCG/Public/Data/PCGDynamicMeshData.h
+++ b/Engine/Plugins/PCG/Source/PCG/Public/Data/PCGDynamicMeshData.h
@@ -64,7 +64,9 @@ public:
 
 	UFUNCTION(BlueprintCallable, Category="DynamicMesh")
 	PCG_API void SetMaterials(const TArray<UMaterialInterface*>& InMaterials);
-	
+
+	PCG_API void SetMaterials(TArray<UMaterialInterface*>&& InMaterials);
+
 	TArray<TObjectPtr<UMaterialInterface>>& GetMutableMaterials() { return Materials; }
 	const TArray<TObjectPtr<UMaterialInterface>>& GetMaterials() const { return Materials; }
 
diff --git a/Engine/Plugins/PCGInterops/PCGGeometryScriptInterop/Source/PCGGeometryScriptInterop/Private/Elements/PCGGetMaterialsDynamicMesh.cpp b/Engine/Plugins/PCGInterops/PCGGeometryScriptInterop/Source/PCGGeometryScriptInterop/Private/Elements/PCGGetMaterialsDynamicMesh.cpp
new file mode 100644
--- /dev/null
+++ b/Engine/Plugins/PCGInterops/PCGGeometryScriptInterop/Source/PCGGeometryScriptInterop/Private/Elements/PCGGetMaterialsDynamicMesh.cpp
@@ -0,0 +1,81 @@
+// Copyright Epic Games, Inc. All Rights Reserved.
+
+#include "Elements/PCGGetMaterialsDynamicMesh.h"
+
+#include "PCGContext.h"
+#include "PCGParamData.h"
+#include "Data/PCGDynamicMeshData.h"
+#include "Helpers/PCGPropertyHelpers.h"
+
+#include "Materials/MaterialInterface.h"
+
+#include UE_INLINE_GENERATED_CPP_BY_NAME(PCGGetMaterialsDynamicMesh)
+
+#define LOCTEXT_NAMESPACE "PCGGetMaterialsDynamicMesh"
+
+#if WITH_EDITOR
+FName UPCGGetMaterialsDynamicMeshSettings::GetDefaultNodeName() const
+{
+	return FName(TEXT("GetMaterialsDynamicMesh"));
+}
+
+FText UPCGGetMaterialsDynamicMeshSettings::GetDefaultNodeTitle() const
+{
+	return LOCTEXT("NodeTitle", "Get Materials On Dynamic Mesh");
+}
+
+FText UPCGGetMaterialsDynamicMeshSettings::GetNodeTooltipText() const
+{
+	return LOCTEXT("NodeTooltip", "Retrieve the array of materials on a dynamic mesh data.");
+}
+#endif // WITH_EDITOR
+
+FPCGElementPtr UPCGGetMaterialsDynamicMeshSettings::CreateElement() const
+{
+	return MakeShared<FPCGGetMaterialsDynamicMeshElement>();
+}
+
+TArray<FPCGPinProperties> UPCGGetMaterialsDynamicMeshSettings::OutputPinProperties() const
+{
+	TArray<FPCGPinProperties> PinProperties;
+	PinProperties.Emplace(PCGPinConstants::DefaultOutputLabel, FPCGDataTypeInfoParam::AsId());
+	return PinProperties;
+}
+
+bool FPCGGetMaterialsDynamicMeshElement::ExecuteInternal(FPCGContext* InContext) const
+{
+	TRACE_CPUPROFILER_EVENT_SCOPE(FPCGGetMaterialsDynamicMeshElement::Execute);
+
+	check(InContext);
+
+	const UPCGGetMaterialsDynamicMeshSettings* Settings = InContext->GetInputSettings<UPCGGetMaterialsDynamicMeshSettings>();
+	check(Settings);
+
+	for (const FPCGTaggedData& Input : InContext->InputData.GetInputsByPin(PCGPinConstants::DefaultInputLabel))
+	{
+		const UPCGDynamicMeshData* InputData = Cast<const UPCGDynamicMeshData>(Input.Data);
+		if (!InputData)
+		{
+			continue;
+		}
+
+		UPCGParamData* OutputData = FPCGContext::NewObject_AnyThread<UPCGParamData>(InContext);
+		UPCGMetadata* OutMetadata = OutputData->MutableMetadata();
+		FPCGMetadataAttribute<FSoftObjectPath>* OutAttribute = OutMetadata->CreateAttribute<FSoftObjectPath>(FName("Materials"), FSoftObjectPath{}, false, false);
+		if (!ensure(OutAttribute))
+		{
+			continue;
+		}
+
+		for (const UMaterialInterface* Material : InputData->GetMaterials())
+		{
+			OutAttribute->SetValue(OutMetadata->AddEntry(), FSoftObjectPath{Material});
+		}
+		
+		InContext->OutputData.TaggedData.Emplace_GetRef(Input).Data = OutputData;
+	}
+
+	return true;
+}
+
+#undef LOCTEXT_NAMESPACE
diff --git a/Engine/Plugins/PCGInterops/PCGGeometryScriptInterop/Source/PCGGeometryScriptInterop/Private/Elements/PCGSetMaterialsDynamicMesh.cpp b/Engine/Plugins/PCGInterops/PCGGeometryScriptInterop/Source/PCGGeometryScriptInterop/Private/Elements/PCGSetMaterialsDynamicMesh.cpp
new file mode 100644
--- /dev/null
+++ b/Engine/Plugins/PCGInterops/PCGGeometryScriptInterop/Source/PCGGeometryScriptInterop/Private/Elements/PCGSetMaterialsDynamicMesh.cpp
@@ -0,0 +1,87 @@
+// Copyright Epic Games, Inc. All Rights Reserved.
+
+#include "Elements/PCGSetMaterialsDynamicMesh.h"
+
+#include "PCGContext.h"
+#include "Data/PCGDynamicMeshData.h"
+
+#include "Materials/MaterialInterface.h"
+
+#include UE_INLINE_GENERATED_CPP_BY_NAME(PCGSetMaterialsDynamicMesh)
+
+#define LOCTEXT_NAMESPACE "PCGSetMaterialsDynamicMesh"
+
+#if WITH_EDITOR
+FName UPCGSetMaterialsDynamicMeshSettings::GetDefaultNodeName() const
+{
+	return FName(TEXT("SetMaterialsDynamicMesh"));
+}
+
+FText UPCGSetMaterialsDynamicMeshSettings::GetDefaultNodeTitle() const
+{
+	return LOCTEXT("NodeTitle", "Set Materials On Dynamic Mesh");
+}
+
+FText UPCGSetMaterialsDynamicMeshSettings::GetNodeTooltipText() const
+{
+	return LOCTEXT("NodeTooltip", "Replace the array of materials on the dynamic mesh data.");
+}
+#endif // WITH_EDITOR
+
+FPCGElementPtr UPCGSetMaterialsDynamicMeshSettings::CreateElement() const
+{
+	return MakeShared<FPCGSetMaterialsDynamicMeshElement>();
+}
+
+bool FPCGSetMaterialsDynamicMeshElement::PrepareDataInternal(FPCGContext* InContext) const
+{
+	TRACE_CPUPROFILER_EVENT_SCOPE(FPCGSetMaterialsDynamicMeshElement::Execute);
+
+	check(InContext);
+
+	const UPCGSetMaterialsDynamicMeshSettings* Settings = InContext->GetInputSettings<UPCGSetMaterialsDynamicMeshSettings>();
+	check(Settings);
+
+	FPCGSetMaterialsDynamicMeshContext* ThisContext = reinterpret_cast<FPCGSetMaterialsDynamicMeshContext*>(InContext);
+
+	TArray<FSoftObjectPath> ResourcesToLoad;
+	ResourcesToLoad.Reserve(Settings->Materials.Num());
+	Algo::TransformIf(
+		Settings->Materials,
+		ResourcesToLoad,
+		[](const TSoftObjectPtr<UMaterialInterface>& Material) { return !Material.IsNull(); },
+		[](const TSoftObjectPtr<UMaterialInterface>& Material) { return Material.ToSoftObjectPath(); }
+	);
+	
+	return ThisContext->RequestResourceLoad(InContext, MoveTemp(ResourcesToLoad));
+}
+
+bool FPCGSetMaterialsDynamicMeshElement::ExecuteInternal(FPCGContext* InContext) const
+{
+	TRACE_CPUPROFILER_EVENT_SCOPE(FPCGSetMaterialsDynamicMeshElement::Execute);
+
+	check(InContext);
+
+	const UPCGSetMaterialsDynamicMeshSettings* Settings = InContext->GetInputSettings<UPCGSetMaterialsDynamicMeshSettings>();
+	check(Settings);
+
+	for (const FPCGTaggedData& Input : InContext->InputData.GetInputsByPin(PCGPinConstants::DefaultInputLabel))
+	{
+		UPCGDynamicMeshData* OutputData = CopyOrSteal(Input, InContext);
+		if (!OutputData)
+		{
+			continue;
+		}
+
+		TArray<UMaterialInterface*> Materials;
+		Materials.Reserve(Settings->Materials.Num());
+		Algo::Transform(Settings->Materials, Materials, [](const TSoftObjectPtr<UMaterialInterface>& Material) { return Material.Get(); });
+
+		OutputData->SetMaterials(MoveTemp(Materials));
+		InContext->OutputData.TaggedData.Emplace_GetRef(Input).Data = OutputData;
+	}
+
+	return true;
+}
+
+#undef LOCTEXT_NAMESPACE
diff --git a/Engine/Plugins/PCGInterops/PCGGeometryScriptInterop/Source/PCGGeometryScriptInterop/Public/Elements/PCGDynamicMeshBaseElement.h b/Engine/Plugins/PCGInterops/PCGGeometryScriptInterop/Source/PCGGeometryScriptInterop/Public/Elements/PCGDynamicMeshBaseElement.h
--- a/Engine/Plugins/PCGInterops/PCGGeometryScriptInterop/Source/PCGGeometryScriptInterop/Public/Elements/PCGDynamicMeshBaseElement.h
+++ b/Engine/Plugins/PCGInterops/PCGGeometryScriptInterop/Source/PCGGeometryScriptInterop/Public/Elements/PCGDynamicMeshBaseElement.h
@@ -43,3 +43,9 @@ public:
 	PCGGEOMETRYSCRIPTINTEROP_API static UPCGDynamicMeshData* CopyOrSteal(const FPCGTaggedData& InTaggedData, FPCGContext* InContext);
 };
 
+template <typename ContextType> requires std::is_base_of_v<FPCGContext, ContextType>
+class IPCGDynamicMeshBaseElementWithCustomContext : public IPCGDynamicMeshBaseElement
+{
+protected:
+	virtual FPCGContext* CreateContext() override { return new ContextType{}; }
+};
diff --git a/Engine/Plugins/PCGInterops/PCGGeometryScriptInterop/Source/PCGGeometryScriptInterop/Public/Elements/PCGGetMaterialsDynamicMesh.h b/Engine/Plugins/PCGInterops/PCGGeometryScriptInterop/Source/PCGGeometryScriptInterop/Public/Elements/PCGGetMaterialsDynamicMesh.h
new file mode 100644
--- /dev/null
+++ b/Engine/Plugins/PCGInterops/PCGGeometryScriptInterop/Source/PCGGeometryScriptInterop/Public/Elements/PCGGetMaterialsDynamicMesh.h
@@ -0,0 +1,36 @@
+﻿// Copyright Epic Games, Inc. All Rights Reserved.
+
+#pragma once
+
+#include "Elements/PCGDynamicMeshBaseElement.h"
+
+#include "PCGGetMaterialsDynamicMesh.generated.h"
+
+/**
+* Retrieve the array of materials on a dynamic mesh data.
+*/
+UCLASS(MinimalAPI, BlueprintType, ClassGroup = (Procedural))
+class UPCGGetMaterialsDynamicMeshSettings : public UPCGDynamicMeshBaseSettings
+{
+	GENERATED_BODY()
+
+public:
+	//~Begin UPCGSettings interface
+#if WITH_EDITOR
+	virtual FName GetDefaultNodeName() const override;
+	virtual FText GetDefaultNodeTitle() const override;
+	virtual FText GetNodeTooltipText() const override;
+#endif
+
+protected:
+	virtual FPCGElementPtr CreateElement() const override;
+	virtual TArray<FPCGPinProperties> OutputPinProperties() const override;
+	//~End UPCGSettings interface
+};
+
+class FPCGGetMaterialsDynamicMeshElement : public IPCGDynamicMeshBaseElement
+{
+protected:
+	virtual bool ExecuteInternal(FPCGContext* InContext) const override;
+};
+
diff --git a/Engine/Plugins/PCGInterops/PCGGeometryScriptInterop/Source/PCGGeometryScriptInterop/Public/Elements/PCGSetMaterialsDynamicMesh.h b/Engine/Plugins/PCGInterops/PCGGeometryScriptInterop/Source/PCGGeometryScriptInterop/Public/Elements/PCGSetMaterialsDynamicMesh.h
new file mode 100644
--- /dev/null
+++ b/Engine/Plugins/PCGInterops/PCGGeometryScriptInterop/Source/PCGGeometryScriptInterop/Public/Elements/PCGSetMaterialsDynamicMesh.h
@@ -0,0 +1,44 @@
+﻿// Copyright Epic Games, Inc. All Rights Reserved.
+
+#pragma once
+
+#include "PCGContext.h"
+#include "Async/PCGAsyncLoadingContext.h"
+#include "Elements/PCGDynamicMeshBaseElement.h"
+
+#include "PCGSetMaterialsDynamicMesh.generated.h"
+
+/**
+* Replace the array of materials on the dynamic mesh data.
+*/
+UCLASS(MinimalAPI, BlueprintType, ClassGroup = (Procedural))
+class UPCGSetMaterialsDynamicMeshSettings : public UPCGDynamicMeshBaseSettings
+{
+	GENERATED_BODY()
+
+public:
+	//~Begin UPCGSettings interface
+#if WITH_EDITOR
+	virtual FName GetDefaultNodeName() const override;
+	virtual FText GetDefaultNodeTitle() const override;
+	virtual FText GetNodeTooltipText() const override;
+#endif
+
+protected:
+	virtual FPCGElementPtr CreateElement() const override;
+	//~End UPCGSettings interface
+	
+public:
+	UPROPERTY(BlueprintReadWrite, EditAnywhere, Category="Settings", meta = (PCG_Overridable))
+	TArray<TSoftObjectPtr<UMaterialInterface>> Materials;
+};
+
+struct FPCGSetMaterialsDynamicMeshContext : public FPCGContext, public IPCGAsyncLoadingContext {};
+
+class FPCGSetMaterialsDynamicMeshElement : public IPCGDynamicMeshBaseElementWithCustomContext<FPCGSetMaterialsDynamicMeshContext>
+{
+protected:
+	virtual bool PrepareDataInternal(FPCGContext* InContext) const override;
+	virtual bool ExecuteInternal(FPCGContext* InContext) const override;
+};
+
