[PCG] Make sure SetMaterialOnDynamicMesh runs on the main thread since it can load objects

#rnx
#Synced-CL 48780673

diff --git a/Engine/Plugins/PCGInterops/PCGGeometryScriptInterop/Source/PCGGeometryScriptInterop/Public/Elements/PCGSetMaterialsDynamicMesh.h b/Engine/Plugins/PCGInterops/PCGGeometryScriptInterop/Source/PCGGeometryScriptInterop/Public/Elements/PCGSetMaterialsDynamicMesh.h
--- a/Engine/Plugins/PCGInterops/PCGGeometryScriptInterop/Source/PCGGeometryScriptInterop/Public/Elements/PCGSetMaterialsDynamicMesh.h
+++ b/Engine/Plugins/PCGInterops/PCGGeometryScriptInterop/Source/PCGGeometryScriptInterop/Public/Elements/PCGSetMaterialsDynamicMesh.h
@@ -37,6 +37,9 @@ struct FPCGSetMaterialsDynamicMeshContext : public FPCGContext, public IPCGAsync
 
 class FPCGSetMaterialsDynamicMeshElement : public IPCGDynamicMeshBaseElementWithCustomContext<FPCGSetMaterialsDynamicMeshContext>
 {
+public:
+	// Loading needs to be done on the main thread and accessing objects outside of PCG might not be thread safe, so taking the safe approach
+	virtual bool CanExecuteOnlyOnMainThread(FPCGContext* Context) const override { return true; }
 protected:
 	virtual bool PrepareDataInternal(FPCGContext* InContext) const override;
 	virtual bool ExecuteInternal(FPCGContext* InContext) const override;
