[PCG] Support for Skeletal Meshes in Dynamic Meshes operations

* Mesh sampler
* Mesh to dynamic mesh
* Append points

#Synced-CL 48852953

diff --git a/Engine/Plugins/PCGInterops/PCGGeometryScriptInterop/Config/DefaultPCGGeometryScriptInterop.ini b/Engine/Plugins/PCGInterops/PCGGeometryScriptInterop/Config/DefaultPCGGeometryScriptInterop.ini
--- a/Engine/Plugins/PCGInterops/PCGGeometryScriptInterop/Config/DefaultPCGGeometryScriptInterop.ini
+++ b/Engine/Plugins/PCGInterops/PCGGeometryScriptInterop/Config/DefaultPCGGeometryScriptInterop.ini
@@ -1,2 +1,7 @@
 ﻿[CoreRedirects]
 +ClassRedirects=(OldName="/Script/PCGGeometryScriptInterop.PCGSplineCrossSectionGeneratorSettings",NewName="/Script/PCGGeometryScriptInterop.PCGPrimitiveCrossSectionSettings")
++ClassRedirects=(OldName="/Script/PCGGeometryScriptInterop.PCGStaticMeshToDynamicMeshSettings",NewName="/Script/PCGGeometryScriptInterop.PCGMeshToDynamicMeshSettings")
+
++PropertyRedirects=(OldName="PCGMeshSamplerSettings.StaticMesh",NewName="PCGMeshSamplerSettings.Mesh")
++PropertyRedirects=(OldName="PCGAppendMeshesFromPointsSettings.StaticMesh",NewName="PCGAppendMeshesFromPointsSettings.Mesh")
++PropertyRedirects=(OldName="PCGMeshToDynamicMeshSettings.StaticMesh",NewName="PCGMeshToDynamicMeshSettings.Mesh")
\ No newline at end of file
diff --git a/Engine/Plugins/PCGInterops/PCGGeometryScriptInterop/Source/PCGGeometryScriptInterop/Private/Elements/PCGAppendMeshesFromPoints.cpp b/Engine/Plugins/PCGInterops/PCGGeometryScriptInterop/Source/PCGGeometryScriptInterop/Private/Elements/PCGAppendMeshesFromPoints.cpp
--- a/Engine/Plugins/PCGInterops/PCGGeometryScriptInterop/Source/PCGGeometryScriptInterop/Private/Elements/PCGAppendMeshesFromPoints.cpp
+++ b/Engine/Plugins/PCGInterops/PCGGeometryScriptInterop/Source/PCGGeometryScriptInterop/Private/Elements/PCGAppendMeshesFromPoints.cpp
@@ -17,6 +17,7 @@
 #include "UDynamicMesh.h"
 #include "ConversionUtils/SceneComponentToDynamicMesh.h"
 #include "DynamicMesh/MeshIndexMappings.h"
+#include "Engine/SkeletalMesh.h"
 #include "Engine/StaticMesh.h"
 #include "GeometryScript/MeshAssetFunctions.h"
 
@@ -49,12 +50,12 @@ FText UPCGAppendMeshesFromPointsSettings::GetNodeTooltipText() const
 
 void UPCGAppendMeshesFromPointsSettings::GetStaticTrackedKeys(FPCGSelectionKeyToSettingsMap& OutKeysToSettings, TArray<TObjectPtr<const UPCGGraph>>& OutVisitedGraphs) const
 {
-	if (Mode != EPCGAppendMeshesFromPointsMode::SingleStaticMesh || StaticMesh.IsNull() || IsPropertyOverriddenByPin(GET_MEMBER_NAME_CHECKED(UPCGAppendMeshesFromPointsSettings, StaticMesh)))
+	if (Mode != EPCGAppendMeshesFromPointsMode::SingleStaticMesh || Mesh.IsNull() || IsPropertyOverriddenByPin(GET_MEMBER_NAME_CHECKED(UPCGAppendMeshesFromPointsSettings, Mesh)))
 	{
 		return;
 	}
 
-	FPCGSelectionKey Key = FPCGSelectionKey::CreateFromPath(StaticMesh.ToSoftObjectPath());
+	FPCGSelectionKey Key = FPCGSelectionKey::CreateFromPath(Mesh.ToSoftObjectPath());
 
 	OutKeysToSettings.FindOrAdd(MoveTemp(Key)).Emplace(this, /*bCulling=*/false);
 }
@@ -127,13 +128,13 @@ bool FPCGAppendMeshesFromPointsElement::PrepareDataInternal(FPCGContext* InConte
 
 	if (Settings->Mode == EPCGAppendMeshesFromPointsMode::SingleStaticMesh)
 	{
-		if (Settings->StaticMesh.IsNull())
+		if (Settings->Mesh.IsNull())
 		{
 			return true;
 		}
 		
 		Context->bPrepareDataSucceeded = true;
-		return Context->RequestResourceLoad(Context, {Settings->StaticMesh.ToSoftObjectPath()}, !Settings->bSynchronousLoad);
+		return Context->RequestResourceLoad(Context, {Settings->Mesh.ToSoftObjectPath()}, !Settings->bSynchronousLoad);
 	}
 
 	if (Settings->Mode == EPCGAppendMeshesFromPointsMode::StaticMeshFromAttribute)
@@ -223,33 +224,61 @@ bool FPCGAppendMeshesFromPointsElement::ExecuteInternal(FPCGContext* InContext)
 	
 	OutputData.Data = OutDynMeshData;
 
-	TMap<FSoftObjectPath, UE::Geometry::FDynamicMesh3> StaticMeshToDynMesh;
+	TMap<FSoftObjectPath, TUniquePtr<UE::Geometry::FDynamicMesh3>> MeshToDynMesh;
 
-	auto ConvertStaticMesh = [Settings, &StaticMeshToDynMesh, InContext, &OutDynMeshData](const TSoftObjectPtr<UStaticMesh>& InMesh) -> bool
+	// Temporary Dynamic Mesh object to be able to use extractor functions for Skeletal Meshes that only have a UDynamicMesh version
+	UDynamicMesh* TempDynMesh = FPCGContext::NewObject_AnyThread<UDynamicMesh>(InContext);
+	UGeometryScriptDebug* Debug = FPCGContext::NewObject_AnyThread<UGeometryScriptDebug>(InContext);
+
+	auto ConvertStaticMesh = [Settings, &MeshToDynMesh, InContext, &OutDynMeshData, &TempDynMesh, Debug](const TSoftObjectPtr<UObject>& InMesh) -> bool
 	{
-		UE::Conversion::FStaticMeshConversionOptions ConversionOptions{};
-		FText ErrorMessage;
-		UE::Conversion::EMeshLODType LODType = static_cast<UE::Conversion::EMeshLODType>(Settings->RequestedLODType);
-		
-		UE::Geometry::FDynamicMesh3& NewMesh = StaticMeshToDynMesh.Add(InMesh.ToSoftObjectPath());
-		UStaticMesh* StaticMesh = InMesh.LoadSynchronous();
-	
-		if (!UE::Conversion::StaticMeshToDynamicMesh(StaticMesh, NewMesh, ErrorMessage, ConversionOptions, LODType, Settings->RequestedLODIndex))
+		const FGeometryScriptMeshReadLOD MeshReadLOD{ Settings->RequestedLODType, Settings->RequestedLODIndex };
+		FGeometryScriptCopyMeshFromAssetOptions Options{};
+
+		UObject* Mesh = InMesh.LoadSynchronous();
+		UStaticMesh* StaticMesh = Cast<UStaticMesh>(Mesh);
+		USkeletalMesh* SkeletalMesh = Cast<USkeletalMesh>(Mesh);
+
+		EGeometryScriptOutcomePins Outcome = EGeometryScriptOutcomePins::Failure;
+
+		if (StaticMesh)
+		{
+			UGeometryScriptLibrary_StaticMeshFunctions::CopyMeshFromStaticMeshV2(StaticMesh, TempDynMesh, Options, MeshReadLOD, Outcome, /*bUseSectionMaterials=*/false, Debug);
+		}
+		else if (SkeletalMesh)
 		{
-			PCGLog::LogErrorOnGraph(ErrorMessage, InContext);
+			UGeometryScriptLibrary_StaticMeshFunctions::CopyMeshFromSkeletalMesh(SkeletalMesh, TempDynMesh, Options, MeshReadLOD, Outcome, Debug);
+		}
+
+		if (Outcome != EGeometryScriptOutcomePins::Success)
+		{
+			PCGLog::LogErrorOnGraph(FText::Format(LOCTEXT("ObjectToDynamicMeshFailed", "Object to Dynamic mesh failed for object {0}. Either it is not a supported type (StaticMesh/SkeletalMesh) or there was a conversion issue."), FText::FromString(InMesh.ToString())));
+			PCGGeometryHelpers::GeometryScriptDebugToPCGLog(InContext, Debug);
 			return false;
 		}
+			
+		TUniquePtr<UE::Geometry::FDynamicMesh3>& NewMesh = MeshToDynMesh.Emplace(InMesh.ToSoftObjectPath(), TempDynMesh->ExtractMesh());
 
 		// Then do the remapping if needed
 		if (Settings->bExtractMaterials)
 		{
-			TArray<UMaterialInterface*> StaticMeshMaterials;
+			TArray<UMaterialInterface*> MeshMaterials;
 			TArray<FName> MaterialSlotNames;
-			UGeometryScriptLibrary_StaticMeshFunctions::GetMaterialListFromStaticMesh(StaticMesh, StaticMeshMaterials, MaterialSlotNames);
 
-			if (!StaticMeshMaterials.IsEmpty() && StaticMeshMaterials != OutDynMeshData->GetMaterials())
+			check(StaticMesh || SkeletalMesh);
+
+			if (StaticMesh)
+			{
+				UGeometryScriptLibrary_StaticMeshFunctions::GetMaterialListFromStaticMesh(StaticMesh, MeshMaterials, MaterialSlotNames);
+			}
+			else
+			{
+				UGeometryScriptLibrary_StaticMeshFunctions::GetMaterialListFromSkeletalMesh(SkeletalMesh, MeshMaterials, MaterialSlotNames);
+			}
+
+			if (!MeshMaterials.IsEmpty() && MeshMaterials != OutDynMeshData->GetMaterials())
 			{
-				PCGGeometryHelpers::RemapMaterials(NewMesh, StaticMeshMaterials, OutDynMeshData->GetMutableMaterials());
+				PCGGeometryHelpers::RemapMaterials(*NewMesh, MeshMaterials, OutDynMeshData->GetMutableMaterials());
 			}
 		}
 
@@ -260,15 +289,15 @@ bool FPCGAppendMeshesFromPointsElement::ExecuteInternal(FPCGContext* InContext)
 	{
 	case EPCGAppendMeshesFromPointsMode::SingleStaticMesh:
 	{
-		if (!ConvertStaticMesh(Settings->StaticMesh))
+		if (!ConvertStaticMesh(Settings->Mesh))
 		{
 			return true;
 		}
 
 #if WITH_EDITOR
-		if (Context->IsValueOverriden(GET_MEMBER_NAME_CHECKED(UPCGAppendMeshesFromPointsSettings, StaticMesh)))
+		if (Context->IsValueOverriden(GET_MEMBER_NAME_CHECKED(UPCGAppendMeshesFromPointsSettings, Mesh)))
 		{
-			FPCGDynamicTrackingHelper::AddSingleDynamicTrackingKey(Context, FPCGSelectionKey::CreateFromPath(Settings->StaticMesh.ToSoftObjectPath()), /*bCulled=*/false);
+			FPCGDynamicTrackingHelper::AddSingleDynamicTrackingKey(Context, FPCGSelectionKey::CreateFromPath(Settings->Mesh.ToSoftObjectPath()), /*bCulled=*/false);
 		}
 #endif // WITH_EDITOR
 
@@ -281,11 +310,11 @@ bool FPCGAppendMeshesFromPointsElement::ExecuteInternal(FPCGContext* InContext)
 		DynamicTracking.EnableAndInitialize(Context, Context->MeshToPointIndicesMapping.Num());
 #endif // WITH_EDITOR
 
-		StaticMeshToDynMesh.Reserve(Context->MeshToPointIndicesMapping.Num());
+		MeshToDynMesh.Reserve(Context->MeshToPointIndicesMapping.Num());
 	
 		for (TPair<FSoftObjectPath, TArray<int32>>& It : Context->MeshToPointIndicesMapping)
 		{
-			if (!ConvertStaticMesh(TSoftObjectPtr<UStaticMesh>(It.Key)))
+			if (!ConvertStaticMesh(TSoftObjectPtr<UObject>(It.Key)))
 			{
 				return true;
 			}
@@ -348,7 +377,7 @@ bool FPCGAppendMeshesFromPointsElement::ExecuteInternal(FPCGContext* InContext)
 	{
 		for (TPair<FSoftObjectPath, TArray<int32>>& It : Context->MeshToPointIndicesMapping)
 		{
-			const UE::Geometry::FDynamicMesh3* MeshToAppend = &StaticMeshToDynMesh[It.Key];
+			const UE::Geometry::FDynamicMesh3* MeshToAppend = MeshToDynMesh[It.Key].Get();
 			check(MeshToAppend);
 
 			for (const int32 i : It.Value)
@@ -362,7 +391,7 @@ bool FPCGAppendMeshesFromPointsElement::ExecuteInternal(FPCGContext* InContext)
 		const UE::Geometry::FDynamicMesh3* MeshToAppend = nullptr;
 		if (Settings->Mode == EPCGAppendMeshesFromPointsMode::SingleStaticMesh)
 		{
-			MeshToAppend = &StaticMeshToDynMesh[Settings->StaticMesh.ToSoftObjectPath()];
+			MeshToAppend = MeshToDynMesh[Settings->Mesh.ToSoftObjectPath()].Get();
 		}
 		else
 		{
diff --git a/Engine/Plugins/PCGInterops/PCGGeometryScriptInterop/Source/PCGGeometryScriptInterop/Private/Elements/PCGMeshSampler.cpp b/Engine/Plugins/PCGInterops/PCGGeometryScriptInterop/Source/PCGGeometryScriptInterop/Private/Elements/PCGMeshSampler.cpp
--- a/Engine/Plugins/PCGInterops/PCGGeometryScriptInterop/Source/PCGGeometryScriptInterop/Private/Elements/PCGMeshSampler.cpp
+++ b/Engine/Plugins/PCGInterops/PCGGeometryScriptInterop/Source/PCGGeometryScriptInterop/Private/Elements/PCGMeshSampler.cpp
@@ -19,6 +19,7 @@
 #include "Data/PCGDynamicMeshData.h"
 #include "Engine/AssetManager.h"
 #include "Engine/Blueprint.h"
+#include "Engine/SkeletalMesh.h"
 #include "Engine/StaticMesh.h"
 #include "GameFramework/Actor.h"
 #include "GeometryScript/MeshAssetFunctions.h"
@@ -356,27 +357,6 @@ UPCGMeshSamplerSettings::UPCGMeshSamplerSettings()
 	}
 }
 
-void UPCGMeshSamplerSettings::PostLoad()
-{
-	Super::PostLoad();
-
-#if WITH_EDITOR
-	if (!StaticMeshPath_DEPRECATED.IsNull())
-	{
-		StaticMesh = StaticMeshPath_DEPRECATED;
-		StaticMeshPath_DEPRECATED.Reset();
-	}
-
-	if (bUseRedAsDensity_DEPRECATED)
-	{
-		// It was only available for one point per vertex before. Keep that.
-		bUseColorChannelAsDensity = (SamplingMethod == EPCGMeshSamplingMethod::OnePointPerVertex);
-		ColorChannelAsDensity = EPCGColorChannel::Red;
-		bUseRedAsDensity_DEPRECATED = false;
-	}
-#endif
-}
-
 TArray<FPCGPinProperties> UPCGMeshSamplerSettings::InputPinProperties() const
 {
 	TArray<FPCGPinProperties> Properties;
@@ -408,12 +388,12 @@ FText UPCGMeshSamplerSettings::GetNodeTooltipText() const
 
 void UPCGMeshSamplerSettings::GetStaticTrackedKeys(FPCGSelectionKeyToSettingsMap& OutKeysToSettings, TArray<TObjectPtr<const UPCGGraph>>& OutVisitedGraphs) const
 {
-	if (StaticMesh.IsNull() || IsPropertyOverriddenByPin(GET_MEMBER_NAME_CHECKED(UPCGMeshSamplerSettings, StaticMesh)))
+	if (Mesh.IsNull() || IsPropertyOverriddenByPin(GET_MEMBER_NAME_CHECKED(UPCGMeshSamplerSettings, Mesh)))
 	{
 		return;
 	}
 
-	FPCGSelectionKey Key = FPCGSelectionKey::CreateFromPath(StaticMesh.ToSoftObjectPath());
+	FPCGSelectionKey Key = FPCGSelectionKey::CreateFromPath(Mesh.ToSoftObjectPath());
 
 	OutKeysToSettings.FindOrAdd(Key).Emplace(this, /*bCulling=*/false);
 }
@@ -468,12 +448,10 @@ bool FPCGMeshSamplerElement::PrepareDataInternal(FPCGContext* InContext) const
 	const UPCGMeshSamplerSettings* Settings = Context->GetInputSettings<UPCGMeshSamplerSettings>();
 	check(Settings);
 
-	const TSoftObjectPtr<UStaticMesh> StaticMeshPtr = Settings->StaticMesh;
-
 	// 1. Request load for meshes/inputs. Return false if we need to wait, otherwise continue.
 	if (!Context->WasLoadRequested())
 	{
-		const bool bNeedToWait = !Context->InitializeAndRequestLoad(PCGPinConstants::DefaultInputLabel, Settings->InputSource, { StaticMeshPtr.ToSoftObjectPath() }, /*bPersistAllData=*/false, /*bSilenceErrorOnEmptyObjectPath=*/true, Settings->bSynchronousLoad);
+		const bool bNeedToWait = !Context->InitializeAndRequestLoad(PCGPinConstants::DefaultInputLabel, Settings->InputSource, { Settings->Mesh.ToSoftObjectPath() }, /*bPersistAllData=*/false, /*bSilenceErrorOnEmptyObjectPath=*/true, Settings->bSynchronousLoad);
 
 		if (bNeedToWait)
 		{
@@ -511,7 +489,7 @@ bool FPCGMeshSamplerElement::PrepareDataInternal(FPCGContext* InContext) const
 
 		TArray<UObject*> Candidates;
 
-		EGeometryScriptOutcomePins Outcome = EGeometryScriptOutcomePins::Success;
+		EGeometryScriptOutcomePins Outcome = EGeometryScriptOutcomePins::Failure;
 		Debug->Messages.Reset();
 
 		// Check if the UObject is a class of an actor. If so, find the all scene components
@@ -535,7 +513,7 @@ bool FPCGMeshSamplerElement::PrepareDataInternal(FPCGContext* InContext) const
 			Candidates.Add(Object);
 		}
 
-		bool bValid = false;
+		bool bDynMeshAdded = false;
 
 		// @todo_pcg: Better support of multiple candidates. For now, take the first valid one.
 		for (UObject* Candidate : Candidates)
@@ -544,6 +522,7 @@ bool FPCGMeshSamplerElement::PrepareDataInternal(FPCGContext* InContext) const
 			{
 				// Adaptation of UGeometryScriptLibrary_SceneUtilityFunctions::CopyMeshFromComponent, since we don't have access to the Material list with this one.
 				Context->DynamicMeshes.Add(FPCGContext::NewObject_AnyThread<UDynamicMesh>(Context));
+				bDynMeshAdded = true;
 				FTransform Transform;
 				FText ErrorMessage;
 				UE::Conversion::FToMeshOptions Options{};
@@ -567,56 +546,71 @@ bool FPCGMeshSamplerElement::PrepareDataInternal(FPCGContext* InContext) const
 				else
 				{
 					Outcome = EGeometryScriptOutcomePins::Success;
-					bValid = true;
 					Context->DynamicMeshes.Last()->SetMesh(std::move(TempDynMesh));
 
 #if WITH_EDITOR
-					// Track the scene component and also the original object, only if the extraction worked
-					DynamicTracking.AddToTracking(FPCGSelectionKey::CreateFromPath(FSoftObjectPath(OriginalObject)), /*bIsCulled=*/false);
-					DynamicTracking.AddToTracking(FPCGSelectionKey::CreateFromPath(FSoftObjectPath(SceneComponent)), /*bIsCulled=*/false);
+					// Track the scene component if is it different from the original object. The original object is tracked outside of this loop.
+					if (OriginalObject != SceneComponent)
+					{
+						DynamicTracking.AddToTracking(FPCGSelectionKey::CreateFromPath(FSoftObjectPath(SceneComponent)), /*bIsCulled=*/false);
+					}
 #endif // WITH_EDITOR
-
-					break;
 				}
 			}
-			else if (UStaticMesh* StaticMesh = Cast<UStaticMesh>(Candidate))
+			else if (Object->IsA<UStaticMesh>() || Object->IsA<USkeletalMesh>())
 			{
 				const FGeometryScriptMeshReadLOD MeshReadLOD{ Settings->RequestedLODType, Settings->RequestedLODIndex };
 
 				Context->DynamicMeshes.Add(FPCGContext::NewObject_AnyThread<UDynamicMesh>(Context));
+				bDynMeshAdded = true;
 				FGeometryScriptCopyMeshFromAssetOptions Options{};
 
-				UGeometryScriptLibrary_StaticMeshFunctions::CopyMeshFromStaticMeshV2(StaticMesh, Context->DynamicMeshes.Last(), Options, MeshReadLOD, Outcome, false, Debug);
-				if (Outcome == EGeometryScriptOutcomePins::Success)
+				if (UStaticMesh* StaticMesh = Cast<UStaticMesh>(Object))
 				{
-					if (Settings->bOutputMaterialInfo)
+					UGeometryScriptLibrary_StaticMeshFunctions::CopyMeshFromStaticMeshV2(StaticMesh, Context->DynamicMeshes.Last(), Options, MeshReadLOD, Outcome, /*bUseSectionMaterials=*/false, Debug);
+					
+					if (Outcome == EGeometryScriptOutcomePins::Success && Settings->bOutputMaterialInfo)
 					{
 						TArray<FName> MaterialSlotNames;
-						UGeometryScriptLibrary_StaticMeshFunctions::GetMaterialListFromStaticMesh(StaticMesh, Context->AssetMaterialList.Emplace_GetRef(), MaterialSlotNames);
+						UGeometryScriptLibrary_StaticMeshFunctions::GetMaterialListFromStaticMesh(StaticMesh, Context->AssetMaterialList.Emplace_GetRef(), MaterialSlotNames, Debug);
+					}
+				}
+				else
+				{
+					USkeletalMesh* SkeletalMesh = CastChecked<USkeletalMesh>(Object);
+				
+					UGeometryScriptLibrary_StaticMeshFunctions::CopyMeshFromSkeletalMesh(SkeletalMesh, Context->DynamicMeshes.Last(), Options, MeshReadLOD, Outcome, Debug);
+					if (Outcome == EGeometryScriptOutcomePins::Success && Settings->bOutputMaterialInfo)
+					{
+						TArray<FName> MaterialSlotNames;
+						UGeometryScriptLibrary_StaticMeshFunctions::GetMaterialListFromSkeletalMesh(SkeletalMesh, Context->AssetMaterialList.Emplace_GetRef(), MaterialSlotNames, Debug);
 					}
-
-#if WITH_EDITOR
-					// Track the static mesh, only if the extraction worked
-					DynamicTracking.AddToTracking(FPCGSelectionKey::CreateFromPath(FSoftObjectPath(StaticMesh)), /*bIsCulled=*/false);
-#endif // WITH_EDITOR
-
-					bValid = true;
-					break;
 				}
 			}
+
+			// @todo_pcg: Better support of multiple candidates. For now, take the first valid one.
+			if (Outcome == EGeometryScriptOutcomePins::Success)
+			{
+				break;
+			}
 		}
 
 		if (Outcome != EGeometryScriptOutcomePins::Success)
 		{
-			PCGLog::LogErrorOnGraph(FText::Format(LOCTEXT("ObjectToDynamicMeshFailed", "Object to Dynamic mesh failed for object {0}."), FText::FromString(Path.ToString())));
+			PCGLog::LogErrorOnGraph(FText::Format(LOCTEXT("ObjectToDynamicMeshFailed", "Object to Dynamic mesh failed for object {0}. Either it is not a supported type (StaticMesh/SkeletalMesh/Actor/SceneComponent) or there was a conversion issue."), FText::FromString(Path.ToString())));
 			PCGGeometryHelpers::GeometryScriptDebugToPCGLog(Context, Debug);
-			Context->DynamicMeshes.Last()->MarkAsGarbage();
-			Context->DynamicMeshes.RemoveAtSwap(Context->DynamicMeshes.Num() - 1);
+			if (bDynMeshAdded)
+			{
+				Context->DynamicMeshes.Last()->MarkAsGarbage();
+				Context->DynamicMeshes.RemoveAtSwap(Context->DynamicMeshes.Num() - 1);
+			}
 		}
-		else if (!bValid)
+		else
 		{
-			PCGE_LOG(Error, GraphAndLog, FText::Format(LOCTEXT("ObjectNotRightType", "Provided object '{0}' is not a supported type. Only supports StaticMesh/Actor/SceneComponent."), FText::FromString(Path.ToString())));
-			continue;
+#if WITH_EDITOR
+			// Track the object, only if the extraction worked
+			DynamicTracking.AddToTracking(FPCGSelectionKey::CreateFromPath(Path), /*bIsCulled=*/false);
+#endif // WITH_EDITOR
 		}
 	}
 
diff --git a/Engine/Plugins/PCGInterops/PCGGeometryScriptInterop/Source/PCGGeometryScriptInterop/Private/Elements/PCGMeshToDynamicMeshElement.cpp b/Engine/Plugins/PCGInterops/PCGGeometryScriptInterop/Source/PCGGeometryScriptInterop/Private/Elements/PCGMeshToDynamicMeshElement.cpp
new file mode 100644
--- /dev/null
+++ b/Engine/Plugins/PCGInterops/PCGGeometryScriptInterop/Source/PCGGeometryScriptInterop/Private/Elements/PCGMeshToDynamicMeshElement.cpp
@@ -0,0 +1,203 @@
+// Copyright Epic Games, Inc. All Rights Reserved.
+
+#include "Elements/PCGMeshToDynamicMeshElement.h"
+
+#include "PCGContext.h"
+#include "PCGModule.h"
+#include "Data/PCGDynamicMeshData.h"
+#include "Helpers/PCGGeometryHelpers.h"
+
+#if WITH_EDITOR
+#include "Helpers/PCGDynamicTrackingHelpers.h"
+#endif // WITH_EDITOR
+
+#include "UDynamicMesh.h"
+#include "Engine/SkeletalMesh.h"
+#include "Engine/StaticMesh.h"
+#include "GeometryScript/MeshAssetFunctions.h"
+
+#include UE_INLINE_GENERATED_CPP_BY_NAME(PCGMeshToDynamicMeshElement)
+
+#define LOCTEXT_NAMESPACE "PCGMeshToDynamicMeshElement"
+
+#if WITH_EDITOR
+FName UPCGMeshToDynamicMeshSettings::GetDefaultNodeName() const
+{
+	return FName(TEXT("MeshToDynamicMeshElement"));
+}
+
+FText UPCGMeshToDynamicMeshSettings::GetDefaultNodeTitle() const
+{
+	return LOCTEXT("NodeTitle", "Mesh To Dynamic Mesh Element");
+}
+
+FText UPCGMeshToDynamicMeshSettings::GetNodeTooltipText() const
+{
+	return LOCTEXT("NodeTooltip", "Convert a static/skeletal mesh into a dynamic mesh data.");
+}
+
+void UPCGMeshToDynamicMeshSettings::GetStaticTrackedKeys(FPCGSelectionKeyToSettingsMap& OutKeysToSettings, TArray<TObjectPtr<const UPCGGraph>>& OutVisitedGraphs) const
+{
+	if (Mesh.IsNull() || IsPropertyOverriddenByPin(GET_MEMBER_NAME_CHECKED(UPCGMeshToDynamicMeshSettings, Mesh)))
+	{
+		return;
+	}
+
+	FPCGSelectionKey Key = FPCGSelectionKey::CreateFromPath(Mesh.ToSoftObjectPath());
+
+	OutKeysToSettings.FindOrAdd(MoveTemp(Key)).Emplace(this, /*bCulling=*/false);
+}
+#endif // WITH_EDITOR
+
+FPCGElementPtr UPCGMeshToDynamicMeshSettings::CreateElement() const
+{
+	return MakeShared<FPCGMeshToDynamicMeshElement>();
+}
+
+TArray<FPCGPinProperties> UPCGMeshToDynamicMeshSettings::InputPinProperties() const
+{
+	return {};
+}
+
+TArray<FPCGPinProperties> UPCGMeshToDynamicMeshSettings::OutputPinProperties() const
+{
+	TArray<FPCGPinProperties> Properties;
+	Properties.Emplace(PCGPinConstants::DefaultOutputLabel, EPCGDataType::DynamicMesh, false, false);
+	return Properties;
+}
+
+bool FPCGMeshToDynamicMeshElement::CanExecuteOnlyOnMainThread(FPCGContext* Context) const
+{
+	// Without context, we can't know, so force it in the main thread to be safe.
+	return !Context || Context->CurrentPhase == EPCGExecutionPhase::PrepareData;
+}
+
+FPCGContext* FPCGMeshToDynamicMeshElement::CreateContext()
+{
+	return new FPCGMeshToDynamicMeshContext();
+}
+
+bool FPCGMeshToDynamicMeshElement::PrepareDataInternal(FPCGContext* InContext) const
+{
+	TRACE_CPUPROFILER_EVENT_SCOPE(FPCGStaticMeshToDynamicMeshElement::Execute);
+
+	FPCGMeshToDynamicMeshContext* Context = static_cast<FPCGMeshToDynamicMeshContext*>(InContext);
+	check(Context);
+
+	const UPCGMeshToDynamicMeshSettings* Settings = InContext->GetInputSettings<UPCGMeshToDynamicMeshSettings>();
+	check(Settings);
+	
+	if (Context->WasLoadRequested() || Settings->Mesh.IsNull())
+	{
+		return true;
+	}
+
+	TArray<FSoftObjectPath> ObjectsToLoad;
+	if (Settings->bExtractMaterials)
+	{
+		Algo::Transform(Settings->OverrideMaterials, ObjectsToLoad, [](const TSoftObjectPtr<UMaterialInterface>& MaterialSoftPtr) { return MaterialSoftPtr.ToSoftObjectPath(); });
+	}
+
+	ObjectsToLoad.Add(Settings->Mesh.ToSoftObjectPath());
+	
+	return Context->RequestResourceLoad(Context, std::move(ObjectsToLoad), !Settings->bSynchronousLoad);
+}
+
+bool FPCGMeshToDynamicMeshElement::ExecuteInternal(FPCGContext* InContext) const
+{
+	TRACE_CPUPROFILER_EVENT_SCOPE(FPCGMeshToDynamicMeshElement::Execute);
+
+	check(InContext);
+
+	const UPCGMeshToDynamicMeshSettings* Settings = InContext->GetInputSettings<UPCGMeshToDynamicMeshSettings>();
+	check(Settings);
+
+	UStaticMesh* StaticMesh = Cast<UStaticMesh>(Settings->Mesh.Get());
+	USkeletalMesh* SkeletalMesh = Cast<USkeletalMesh>(Settings->Mesh.Get());
+	
+	if (!StaticMesh && !SkeletalMesh)
+	{
+		if (!Settings->Mesh.IsNull())
+		{
+			PCGLog::LogErrorOnGraph(FText::Format(LOCTEXT("MeshNull", "Object {0} failed to load or was not a Static Mesh nor Skeletal Mesh."), FText::FromString(Settings->Mesh.ToSoftObjectPath().ToString())), InContext);
+		}
+		
+		return true;
+	}
+
+	check(StaticMesh || SkeletalMesh);
+
+#if WITH_EDITOR
+	if (InContext->IsValueOverriden(GET_MEMBER_NAME_CHECKED(UPCGMeshToDynamicMeshSettings, Mesh)))
+	{
+		FPCGDynamicTrackingHelper::AddSingleDynamicTrackingKey(InContext, FPCGSelectionKey::CreateFromPath(Settings->Mesh.ToSoftObjectPath()), /*bIsCulled=*/false);
+	}
+#endif // WITH_EDITOR
+
+	if (Settings->bExtractMaterials && !Settings->OverrideMaterials.IsEmpty())
+	{
+		const int32 NumMeshMaterials = StaticMesh ? StaticMesh->GetStaticMaterials().Num() : SkeletalMesh->GetMaterials().Num();
+		if (NumMeshMaterials != Settings->OverrideMaterials.Num())
+		{
+			PCGLog::LogErrorOnGraph(FText::Format(LOCTEXT("MismatchMaterials", "Mismatch number between Mesh materials ({0}) and override materials ({1})"), NumMeshMaterials, Settings->OverrideMaterials.Num()));
+			return true;
+		}
+		else if (const TSoftObjectPtr<UMaterialInterface>* UnloadedMaterial = Settings->OverrideMaterials.FindByPredicate([](const TSoftObjectPtr<UMaterialInterface>& MaterialSoftPtr) -> bool { return !MaterialSoftPtr.Get(); }))
+		{
+			PCGLog::LogErrorOnGraph(FText::Format(LOCTEXT("UnloadedMaterial", "Material {0} failed to load."), FText::FromString(UnloadedMaterial->ToSoftObjectPath().ToString())));
+			return true;
+		}
+	}
+
+	UPCGDynamicMeshData* DynMeshData = FPCGContext::NewObject_AnyThread<UPCGDynamicMeshData>(InContext);
+	
+	const FGeometryScriptMeshReadLOD MeshReadLOD{ Settings->RequestedLODType, Settings->RequestedLODIndex };
+	FGeometryScriptCopyMeshFromAssetOptions Options{};
+
+	UGeometryScriptDebug* Debug = FPCGContext::NewObject_AnyThread<UGeometryScriptDebug>(InContext);
+	EGeometryScriptOutcomePins Outcome = EGeometryScriptOutcomePins::Success;
+	
+	if (StaticMesh)
+	{
+		UGeometryScriptLibrary_StaticMeshFunctions::CopyMeshFromStaticMeshV2(StaticMesh, DynMeshData->GetMutableDynamicMesh(), Options, MeshReadLOD, Outcome, /*bUseSectionMaterials=*/false, Debug);
+	}
+	else
+	{
+		UGeometryScriptLibrary_StaticMeshFunctions::CopyMeshFromSkeletalMesh(SkeletalMesh, DynMeshData->GetMutableDynamicMesh(), Options, MeshReadLOD, Outcome, Debug);
+	}
+
+	if (Outcome == EGeometryScriptOutcomePins::Success)
+	{
+		if (Settings->bExtractMaterials)
+		{
+			TArray<UMaterialInterface*> Materials;
+			
+			if (!Settings->OverrideMaterials.IsEmpty())
+			{
+				Algo::Transform(Settings->OverrideMaterials, Materials, [](const TSoftObjectPtr<UMaterialInterface>& MaterialSoftPtr) { return MaterialSoftPtr.Get(); });
+			}
+			else if (StaticMesh)
+			{
+				TArray<FName> MaterialSlotNames;
+				UGeometryScriptLibrary_StaticMeshFunctions::GetMaterialListFromStaticMesh(StaticMesh, Materials, MaterialSlotNames);
+			}
+			else
+			{
+				TArray<FName> MaterialSlotNames;
+				UGeometryScriptLibrary_StaticMeshFunctions::GetMaterialListFromSkeletalMesh(SkeletalMesh, Materials, MaterialSlotNames);
+			}
+
+			DynMeshData->SetMaterials(MoveTemp(Materials));
+		}
+
+		InContext->OutputData.TaggedData.Emplace_GetRef().Data = DynMeshData;
+	}
+	else
+	{
+		PCGGeometryHelpers::GeometryScriptDebugToPCGLog(InContext, Debug);
+	}
+
+	return true;
+}
+
+#undef LOCTEXT_NAMESPACE
diff --git a/Engine/Plugins/PCGInterops/PCGGeometryScriptInterop/Source/PCGGeometryScriptInterop/Private/Elements/PCGStaticMeshToDynamicMeshElement.cpp b/Engine/Plugins/PCGInterops/PCGGeometryScriptInterop/Source/PCGGeometryScriptInterop/Private/Elements/PCGStaticMeshToDynamicMeshElement.cpp
deleted file mode 100644
diff --git a/Engine/Plugins/PCGInterops/PCGGeometryScriptInterop/Source/PCGGeometryScriptInterop/Public/Elements/PCGAppendMeshesFromPoints.h b/Engine/Plugins/PCGInterops/PCGGeometryScriptInterop/Source/PCGGeometryScriptInterop/Public/Elements/PCGAppendMeshesFromPoints.h
--- a/Engine/Plugins/PCGInterops/PCGGeometryScriptInterop/Source/PCGGeometryScriptInterop/Public/Elements/PCGAppendMeshesFromPoints.h
+++ b/Engine/Plugins/PCGInterops/PCGGeometryScriptInterop/Source/PCGGeometryScriptInterop/Public/Elements/PCGAppendMeshesFromPoints.h
@@ -13,13 +13,13 @@
 UENUM(Blueprintable)
 enum class EPCGAppendMeshesFromPointsMode : uint8
 {
-	SingleStaticMesh UMETA(ToolTip="Mesh taken from the node settings"),
-	StaticMeshFromAttribute UMETA(ToolTip="Mesh taken from attributes on the points"),
-	DynamicMesh UMETA(ToolTip="Mesh taken from another dynamic mesh")
+	SingleStaticMesh UMETA(DisplayName="SingleMesh", ToolTip="Mesh taken from the node settings"),
+	StaticMeshFromAttribute UMETA(DisplayName="MeshFromAttribute", ToolTip="Mesh taken from attributes on the points"),
+	DynamicMesh UMETA(ToolTip="Mesh taken from another dynamic mesh"),
 };
 
 /**
-* Append meshes at the points transforms. Mesh can be a single static mesh, multiple meshes coming from the points or another dynamic mesh.
+* Append meshes at the points transforms. Mesh can be a single static/skeletal mesh, multiple meshes coming from the points or another dynamic mesh.
 */
 UCLASS(MinimalAPI, BlueprintType, ClassGroup = (Procedural))
 class UPCGAppendMeshesFromPointsSettings : public UPCGDynamicMeshBaseSettings
@@ -49,17 +49,17 @@ public:
 	UPROPERTY(BlueprintReadWrite, EditAnywhere, Category="Settings")
 	EPCGAppendMeshesFromPointsMode Mode;
 	
-	UPROPERTY(BlueprintReadWrite, EditAnywhere, Category="Settings", meta = (EditCondition = "Mode==EPCGAppendMeshesFromPointsMode::SingleStaticMesh", EditConditionHides, PCG_Overridable))
-	TSoftObjectPtr<UStaticMesh> StaticMesh;
+	UPROPERTY(BlueprintReadWrite, EditAnywhere, Category = Settings, meta = (EditCondition = "Mode==EPCGAppendMeshesFromPointsMode::SingleStaticMesh", EditConditionHides, PCG_Overridable, PCG_OverrideAliases="StaticMesh", AllowedClasses="/Script/Engine.StaticMesh, /Script/Engine.SkeletalMesh"))
+	TSoftObjectPtr<UObject> Mesh;
 
 	UPROPERTY(BlueprintReadWrite, EditAnywhere, Category="Settings", meta = (EditCondition = "Mode==EPCGAppendMeshesFromPointsMode::StaticMeshFromAttribute", EditConditionHides, PCG_Overridable))
 	FPCGAttributePropertyInputSelector MeshAttribute;
 	
-	/** Allows to extract materials from the static mesh and set them in the resulting append. */
+	/** Allows to extract materials from the static/skeletal mesh and set them in the resulting append. */
     UPROPERTY(BlueprintReadWrite, EditAnywhere, Category="Settings", meta = (EditCondition = "Mode!=EPCGAppendMeshesFromPointsMode::DynamicMesh", EditConditionHides, PCG_Overridable))
     bool bExtractMaterials = true;
 	
-	/** LOD type to use when creating DynamicMesh from specified StaticMesh. */
+	/** LOD type to use when creating DynamicMesh from specified StaticMesh/SkeletalMesh. */
 	UPROPERTY(BlueprintReadWrite, EditAnywhere, Category = "Settings|LODSettings", meta = (EditCondition = "Mode!=EPCGAppendMeshesFromPointsMode::DynamicMesh", EditConditionHides, PCG_Overridable))
 	EGeometryScriptLODType RequestedLODType = EGeometryScriptLODType::RenderData;
 
@@ -68,6 +68,9 @@ public:
 	
 	UPROPERTY(BlueprintReadWrite, EditAnywhere, Category = "Settings|Debug", meta = (EditCondition = "Mode!=EPCGAppendMeshesFromPointsMode::DynamicMesh", EditConditionHides))
 	bool bSynchronousLoad = false;
+
+	UE_DEPRECATED(5.8, "Use Mesh property")
+	TSoftObjectPtr<UStaticMesh> StaticMesh;
 };
 
 struct FPCGAppendMeshesFromPointsContext : public FPCGContext, public IPCGAsyncLoadingContext
diff --git a/Engine/Plugins/PCGInterops/PCGGeometryScriptInterop/Source/PCGGeometryScriptInterop/Public/Elements/PCGMeshSampler.h b/Engine/Plugins/PCGInterops/PCGGeometryScriptInterop/Source/PCGGeometryScriptInterop/Public/Elements/PCGMeshSampler.h
--- a/Engine/Plugins/PCGInterops/PCGGeometryScriptInterop/Source/PCGGeometryScriptInterop/Public/Elements/PCGMeshSampler.h
+++ b/Engine/Plugins/PCGInterops/PCGGeometryScriptInterop/Source/PCGGeometryScriptInterop/Public/Elements/PCGMeshSampler.h
@@ -61,10 +61,6 @@ class UPCGMeshSamplerSettings : public UPCGSettings
 public:
 	UE_API UPCGMeshSamplerSettings();
 
-	//~Begin UObject interface
-	UE_API virtual void PostLoad() override;
-	//~End UObject interface
-
 	//~Begin UPCGSettings interface
 #if WITH_EDITOR
 	virtual FName GetDefaultNodeName() const override { return FName(TEXT("MeshSampler")); }
@@ -84,16 +80,16 @@ protected:
 
 public:
 	/** 
-	* Can provide a list of inputs to sample the meshes from. It can be a list of StaticMeshes, a list of Actors that have a Scene Component (like a Static Mesh Component), or a list of Scene Components directly.
+	* Can provide a list of inputs to sample the meshes from. It can be a list of StaticMeshes/SkeletalMeshes, a list of Actors that have a Scene Component (like a Static Mesh Component), or a list of Scene Components directly.
 	* Geometry Script needs to be able to extract a dynamic mesh from this scene component (so won't work for ISMCs for example) and for now will work only with a single scene component.
-	* Each entry (either in the same data or seperate data) will produce a unique output data.
+	* Each entry (either in the same data or separate data) will produce a unique output data.
 	*/
 	UPROPERTY(BlueprintReadWrite, EditAnywhere, Category = "Settings")
 	bool bExtractMeshFromInput = false;
 
-	/** Soft Object Path to the mesh to sample from. Will be loaded. */
-	UPROPERTY(BlueprintReadWrite, EditAnywhere, Category = "Settings", meta = (PCG_Overridable, EditCondition = "!bExtractMeshFromInput", EditConditionHides, ShowAfter = "bExtractMeshFromInput"))
-	TSoftObjectPtr<UStaticMesh> StaticMesh;
+	/** Soft Object Path to the mesh to sample from. Supports Static Meshes and Skeletal Meshes. Will be loaded. */
+	UPROPERTY(BlueprintReadWrite, EditAnywhere, Category = Settings, meta = (PCG_Overridable, PCG_OverrideAliases="StaticMesh", AllowedClasses="/Script/Engine.StaticMesh, /Script/Engine.SkeletalMesh", EditCondition = "!bExtractMeshFromInput", EditConditionHides, ShowAfter = "bExtractMeshFromInput"))
+	TSoftObjectPtr<UObject> Mesh;
 
 	/** Selector to read data from. */
 	UPROPERTY(BlueprintReadWrite, EditAnywhere, Category = "Settings", meta = (PCG_Overridable, EditCondition = "bExtractMeshFromInput", EditConditionHides, ShowAfter = "bExtractMeshFromInput"))
@@ -121,7 +117,7 @@ public:
 	UPROPERTY(BlueprintReadWrite, EditAnywhere, Category = "Settings|Voxelize Options", meta = (PCG_Overridable, EditCondition = "bVoxelize"))
 	bool bRemoveHiddenTriangles = true;
 
-	/** LOD type to use when creating DynamicMesh from specified StaticMesh. */
+	/** LOD type to use when creating DynamicMesh from specified StaticMesh/SkeletalMesh. */
 	UPROPERTY(BlueprintReadWrite, EditAnywhere, Category = "Settings|LODSettings", meta = (PCG_Overridable))
 	EGeometryScriptLODType RequestedLODType = EGeometryScriptLODType::RenderData;
 
@@ -169,16 +165,9 @@ public:
 
 	UPROPERTY(BlueprintReadWrite, EditAnywhere, Category = "Settings|Debug")
 	bool bSynchronousLoad = false;
-	
-protected:
-#if WITH_EDITORONLY_DATA
-	// Deprecated in UE 5.3 in favor of StaticMesh
-	UPROPERTY()
-	FSoftObjectPath StaticMeshPath_DEPRECATED;
 
-	UPROPERTY()
-	bool bUseRedAsDensity_DEPRECATED = false;
-#endif
+	UE_DEPRECATED(5.8, "Use Mesh property")
+	TSoftObjectPtr<UStaticMesh> StaticMesh;
 };
 
 /**
diff --git a/Engine/Plugins/PCGInterops/PCGGeometryScriptInterop/Source/PCGGeometryScriptInterop/Public/Elements/PCGMeshToDynamicMeshElement.h b/Engine/Plugins/PCGInterops/PCGGeometryScriptInterop/Source/PCGGeometryScriptInterop/Public/Elements/PCGMeshToDynamicMeshElement.h
new file mode 100644
--- /dev/null
+++ b/Engine/Plugins/PCGInterops/PCGGeometryScriptInterop/Source/PCGGeometryScriptInterop/Public/Elements/PCGMeshToDynamicMeshElement.h
@@ -0,0 +1,74 @@
+// Copyright Epic Games, Inc. All Rights Reserved.
+
+#pragma once
+
+#include "PCGContext.h"
+#include "Async/PCGAsyncLoadingContext.h"
+#include "Elements/PCGDynamicMeshBaseElement.h"
+
+#include "GeometryScript/GeometryScriptTypes.h"
+
+#include "PCGMeshToDynamicMeshElement.generated.h"
+
+/**
+* Convert a static/skeletal mesh into a dynamic mesh data.
+*/
+UCLASS(MinimalAPI, BlueprintType, ClassGroup = (Procedural))
+class UPCGMeshToDynamicMeshSettings : public UPCGDynamicMeshBaseSettings
+{
+	GENERATED_BODY()
+
+public:
+	//~Begin UPCGSettings interface
+#if WITH_EDITOR
+	virtual FName GetDefaultNodeName() const override;
+	virtual FText GetDefaultNodeTitle() const override;
+	virtual FText GetNodeTooltipText() const override;
+	virtual void GetStaticTrackedKeys(FPCGSelectionKeyToSettingsMap& OutKeysToSettings, TArray<TObjectPtr<const UPCGGraph>>& OutVisitedGraphs) const override;
+	virtual bool CanDynamicallyTrackKeys() const override { return true; };
+#endif
+
+protected:
+	virtual TArray<FPCGPinProperties> InputPinProperties() const override;
+	virtual TArray<FPCGPinProperties> OutputPinProperties() const override;
+	virtual FPCGElementPtr CreateElement() const override;
+	//~End UPCGSettings interface
+
+public:
+	UPROPERTY(BlueprintReadWrite, EditAnywhere, Category = Settings, meta = (PCG_Overridable, PCG_OverrideAliases="StaticMesh", AllowedClasses="/Script/Engine.StaticMesh, /Script/Engine.SkeletalMesh"))
+	TSoftObjectPtr<UObject> Mesh;
+
+	/** Allows to extract materials from the static/skeletal mesh and store them in the PCG Dynamic Mesh Data. */
+	UPROPERTY(BlueprintReadWrite, EditAnywhere, Category="Settings", meta = (PCG_Overridable))
+	bool bExtractMaterials = true;
+
+	/** If it extracts materials, we can specify override materials. It needs to have the same number of material overrides than there are materials on the static mesh. */
+	UPROPERTY(EditAnywhere, Category = "Settings", meta = (EditCondition = "bExtractMaterials", EditConditionHides, PCG_Overridable))
+	TArray<TSoftObjectPtr<UMaterialInterface>> OverrideMaterials;
+
+	/** LOD type to use when creating DynamicMesh from specified StaticMesh/SkeletalMesh. */
+	UPROPERTY(BlueprintReadWrite, EditAnywhere, Category = "Settings|LODSettings", meta = (PCG_Overridable))
+	EGeometryScriptLODType RequestedLODType = EGeometryScriptLODType::MaxAvailable;
+
+	UPROPERTY(BlueprintReadWrite, EditAnywhere, Category = "Settings|LODSettings", meta = (PCG_Overridable))
+	int32 RequestedLODIndex = 0;
+	
+	UPROPERTY(BlueprintReadWrite, EditAnywhere, Category = "Settings|Debug")
+	bool bSynchronousLoad = false;
+
+	UE_DEPRECATED(5.8, "Use Mesh property")
+	TSoftObjectPtr<UStaticMesh> StaticMesh;
+};
+
+struct FPCGMeshToDynamicMeshContext : public FPCGContext, public IPCGAsyncLoadingContext {};
+
+class FPCGMeshToDynamicMeshElement : public IPCGDynamicMeshBaseElement
+{
+public:
+	virtual bool CanExecuteOnlyOnMainThread(FPCGContext* Context) const override;
+protected:
+	virtual FPCGContext* CreateContext() override;
+	virtual bool PrepareDataInternal(FPCGContext* Context) const override;
+	virtual bool ExecuteInternal(FPCGContext* InContext) const override;
+};
+
diff --git a/Engine/Plugins/PCGInterops/PCGGeometryScriptInterop/Source/PCGGeometryScriptInterop/Public/Elements/PCGStaticMeshToDynamicMeshElement.h b/Engine/Plugins/PCGInterops/PCGGeometryScriptInterop/Source/PCGGeometryScriptInterop/Public/Elements/PCGStaticMeshToDynamicMeshElement.h
--- a/Engine/Plugins/PCGInterops/PCGGeometryScriptInterop/Source/PCGGeometryScriptInterop/Public/Elements/PCGStaticMeshToDynamicMeshElement.h
+++ b/Engine/Plugins/PCGInterops/PCGGeometryScriptInterop/Source/PCGGeometryScriptInterop/Public/Elements/PCGStaticMeshToDynamicMeshElement.h
@@ -2,72 +2,6 @@
 
 #pragma once
 
-#include "PCGContext.h"
-#include "Async/PCGAsyncLoadingContext.h"
-#include "Elements/PCGDynamicMeshBaseElement.h"
-
-#include "GeometryScript/GeometryScriptTypes.h"
-
-#include "PCGStaticMeshToDynamicMeshElement.generated.h"
-
-class UStaticMesh;
-
-/**
-* Convert a static mesh into a dynamic mesh data.
-*/
-UCLASS(MinimalAPI, BlueprintType, ClassGroup = (Procedural))
-class UPCGStaticMeshToDynamicMeshSettings : public UPCGDynamicMeshBaseSettings
-{
-	GENERATED_BODY()
-
-public:
-	//~Begin UPCGSettings interface
-#if WITH_EDITOR
-	virtual FName GetDefaultNodeName() const override;
-	virtual FText GetDefaultNodeTitle() const override;
-	virtual FText GetNodeTooltipText() const override;
-	virtual void GetStaticTrackedKeys(FPCGSelectionKeyToSettingsMap& OutKeysToSettings, TArray<TObjectPtr<const UPCGGraph>>& OutVisitedGraphs) const override;
-	virtual bool CanDynamicallyTrackKeys() const override { return true; };
-#endif
-
-protected:
-	virtual TArray<FPCGPinProperties> InputPinProperties() const override;
-	virtual TArray<FPCGPinProperties> OutputPinProperties() const override;
-	virtual FPCGElementPtr CreateElement() const override;
-	//~End UPCGSettings interface
-
-public:
-	UPROPERTY(BlueprintReadWrite, EditAnywhere, Category="Settings", meta = (PCG_Overridable))
-	TSoftObjectPtr<UStaticMesh> StaticMesh;
-
-	/** Allows to extract materials from the static mesh and store them in the PCG Dynamic Mesh Data. */
-	UPROPERTY(BlueprintReadWrite, EditAnywhere, Category="Settings", meta = (PCG_Overridable))
-	bool bExtractMaterials = true;
-
-	/** If it extracts materials, we can specify override materials. It needs to have the same number of material overrides than there are materials on the static mesh. */
-	UPROPERTY(EditAnywhere, Category = "Settings", meta = (EditCondition = "bExtractMaterials", EditConditionHides))
-	TArray<TSoftObjectPtr<UMaterialInterface>> OverrideMaterials;
-
-	/** LOD type to use when creating DynamicMesh from specified StaticMesh. */
-	UPROPERTY(BlueprintReadWrite, EditAnywhere, Category = "Settings|LODSettings", meta = (PCG_Overridable))
-	EGeometryScriptLODType RequestedLODType = EGeometryScriptLODType::MaxAvailable;
-
-	UPROPERTY(BlueprintReadWrite, EditAnywhere, Category = "Settings|LODSettings", meta = (PCG_Overridable))
-	int32 RequestedLODIndex = 0;
-	
-	UPROPERTY(BlueprintReadWrite, EditAnywhere, Category = "Settings|Debug")
-	bool bSynchronousLoad = false;
-};
-
-struct FPCGStaticMeshToDynamicMeshContext : public FPCGContext, public IPCGAsyncLoadingContext {};
-
-class FPCGStaticMeshToDynamicMeshElement : public IPCGDynamicMeshBaseElement
-{
-public:
-	virtual bool CanExecuteOnlyOnMainThread(FPCGContext* Context) const override;
-protected:
-	virtual FPCGContext* CreateContext() override;
-	virtual bool PrepareDataInternal(FPCGContext* Context) const override;
-	virtual bool ExecuteInternal(FPCGContext* InContext) const override;
-};
+UE_DEPRECATED_HEADER(5.8, "Include Elements/PCGMeshToDynamicMeshElement.h")
 
+#include "Elements/PCGMeshToDynamicMeshElement.h"
\ No newline at end of file
