[PCG] Made spline to mesh polygon aware

#prefliight 69a08242764024eb15bdcb54
#Synced-CL 51118381

diff --git a/Engine/Plugins/PCGInterops/PCGGeometryScriptInterop/Source/PCGGeometryScriptInterop/Private/Elements/PCGSplineToMesh.cpp b/Engine/Plugins/PCGInterops/PCGGeometryScriptInterop/Source/PCGGeometryScriptInterop/Private/Elements/PCGSplineToMesh.cpp
--- a/Engine/Plugins/PCGInterops/PCGGeometryScriptInterop/Source/PCGGeometryScriptInterop/Private/Elements/PCGSplineToMesh.cpp
+++ b/Engine/Plugins/PCGInterops/PCGGeometryScriptInterop/Source/PCGGeometryScriptInterop/Private/Elements/PCGSplineToMesh.cpp
@@ -5,6 +5,7 @@
 #include "PCGContext.h"
 #include "Data/PCGDynamicMeshData.h"
 #include "Data/PCGSplineData.h"
+#include "Data/PCGPolygon2DData.h"
 
 #include "CurveOps/TriangulateCurvesOp.h"
 #include "Operations/ExtrudeMesh.h"
@@ -13,6 +14,11 @@
 
 #define LOCTEXT_NAMESPACE "PCGSplineToMeshElement"
 
+namespace PCGSplineToMeshConstants
+{
+	const FLazyName MeshTransform = "MeshTransform";
+}
+
 #if WITH_EDITOR
 FName UPCGSplineToMeshSettings::GetDefaultNodeName() const
 {
@@ -38,7 +44,7 @@ FPCGElementPtr UPCGSplineToMeshSettings::CreateElement() const
 TArray<FPCGPinProperties> UPCGSplineToMeshSettings::InputPinProperties() const
 {
 	TArray<FPCGPinProperties> Properties;
-	Properties.Emplace_GetRef(PCGPinConstants::DefaultInputLabel, EPCGDataType::Spline).SetRequiredPin();	
+	Properties.Emplace_GetRef(PCGPinConstants::DefaultInputLabel, EPCGDataType::Spline | EPCGDataType::Polygon2D).SetRequiredPin();
 	return Properties;
 }
 
@@ -56,10 +62,12 @@ bool FPCGSplineToMeshElement::ExecuteInternal(FPCGContext* InContext) const
 	const TArray<FPCGTaggedData> Inputs = InContext->InputData.GetInputsByPin(PCGPinConstants::DefaultInputLabel);
 	for (const FPCGTaggedData& Input : Inputs)
 	{
-		const UPCGSplineData* InputSplineData = Cast<const UPCGSplineData>(Input.Data);
-		if (!InputSplineData)
+		const UPCGSplineData* InputSplineData = Cast<UPCGSplineData>(Input.Data);
+		const UPCGPolygon2DData* InputPolygonData = Cast<UPCGPolygon2DData>(Input.Data);
+
+		if (!InputSplineData && !InputPolygonData)
 		{
-			PCGLog::InputOutput::LogTypedDataNotFoundWarning(EPCGDataType::Spline, PCGPinConstants::DefaultInputLabel, InContext);
+			PCGLog::InputOutput::LogTypedDataNotFoundWarning(EPCGDataType::Spline | EPCGDataType::Polygon2D, PCGPinConstants::DefaultInputLabel, InContext);
 			continue;
 		}
 
@@ -83,9 +91,43 @@ bool FPCGSplineToMeshElement::ExecuteInternal(FPCGContext* InContext) const
 			TriangulateCurvesOp.OffsetClosedMethod = Settings->OffsetClosedCurves;
 		}
 		
-		TArray<FVector> SplinePoints;
-		InputSplineData->SplineStruct.ConvertSplineToPolyLine(ESplineCoordinateSpace::World, Settings->ErrorTolerance, SplinePoints);
-		TriangulateCurvesOp.AddWorldCurve(SplinePoints, InputSplineData->IsClosed(), InputSplineData->SplineStruct.Transform);
+		if (InputSplineData)
+		{
+			TArray<FVector> SplinePoints;
+			InputSplineData->SplineStruct.ConvertSplineToPolyLine(ESplineCoordinateSpace::World, Settings->ErrorTolerance, SplinePoints);
+			TriangulateCurvesOp.AddWorldCurve(SplinePoints, InputSplineData->IsClosed(), InputSplineData->SplineStruct.Transform);
+		}
+		else
+		{
+			check(InputPolygonData);
+			const UE::Geometry::FGeneralPolygon2d& Polygon = InputPolygonData->GetPolygon();
+			const FTransform Transform = InputPolygonData->GetTransform();
+
+			TArray<FVector> PolyVertices;
+			auto GetVerticesFromPolygon =  [&Transform](const UE::Geometry::FPolygon2d& Poly, TArray<FVector>& PolyVertices)
+			{
+				const TArray<FVector2D>& PolyVertices2D = Poly.GetVertices();
+
+				PolyVertices.Reset();
+				PolyVertices.Reserve(PolyVertices2D.Num());
+				Algo::Transform(PolyVertices2D, PolyVertices, [&Transform](const FVector2D& Vtx2D) { return Transform.TransformPosition(FVector(Vtx2D, 0.0)); });
+			};
+
+			GetVerticesFromPolygon(Polygon.GetOuter(), PolyVertices);
+			TriangulateCurvesOp.AddWorldCurve(PolyVertices, /*bClosed=*/true, Transform);
+
+			for (const UE::Geometry::FPolygon2d& Hole : Polygon.GetHoles())
+			{
+				GetVerticesFromPolygon(Hole, PolyVertices);
+				TriangulateCurvesOp.AddWorldCurve(PolyVertices, /*bClosed=*/true, Transform);
+			}
+
+			// Implementation note:
+			// It is imperative that the curves are flattened, otherwise the difference will not happen.
+			TriangulateCurvesOp.CombineMethod = ECombineCurvesMethod::Difference;
+			TriangulateCurvesOp.FlattenMethod = EFlattenCurveMethod::ToBestFitPlane;
+		}
+
 		TriangulateCurvesOp.CalculateResult(nullptr);
 
 		TUniquePtr<UE::Geometry::FDynamicMesh3>	DynamicMesh = TriangulateCurvesOp.ExtractResult();
@@ -120,6 +162,15 @@ bool FPCGSplineToMeshElement::ExecuteInternal(FPCGContext* InContext) const
 		delete DynamicMeshPtr;
 		DynamicMeshPtr = nullptr;
 
+		UPCGMetadata* Metadata = DynamicMeshData->MutableMetadata();
+		if (FPCGMetadataDomain* DataDomain = Metadata->GetMetadataDomain(PCGMetadataDomainID::Data))
+		{
+			FPCGMetadataAttribute<FTransform>* TransformAttribute = DataDomain->CreateAttribute<FTransform>(PCGSplineToMeshConstants::MeshTransform, FTransform::Identity, true, true);
+			check(TransformAttribute);
+			PCGMetadataEntryKey DefaultDomainKey = DataDomain->AddEntry();
+			TransformAttribute->SetValue(DefaultDomainKey, TriangulateCurvesOp.GetResultTransform());
+		}
+
 		InContext->OutputData.TaggedData.Emplace_GetRef(Input).Data = DynamicMeshData;
 	}
 	
