[PCG] Add default color override for Mesh Sampler + fix bounds in SubdivideSpline

* Bundle them together with a version bump for deprecation

#rnx
#Synced-CL 51771932

diff --git a/Engine/Plugins/PCG/Source/PCG/Private/Elements/Grammar/PCGSubdivideSpline.cpp b/Engine/Plugins/PCG/Source/PCG/Private/Elements/Grammar/PCGSubdivideSpline.cpp
--- a/Engine/Plugins/PCG/Source/PCG/Private/Elements/Grammar/PCGSubdivideSpline.cpp
+++ b/Engine/Plugins/PCG/Source/PCG/Private/Elements/Grammar/PCGSubdivideSpline.cpp
@@ -103,6 +103,7 @@ namespace PCGSubdivideSplineHelpers
 		}
 
 		const double SplineLength = InOutParameters.PolyLineData->GetLength();
+		const double HalfModuleHeight = InOutParameters.ModuleHeight / 2;
 
 		/* Implementation Note: Subdivided spline length will always be equal or greater than discretized linear length, depending on the curvature of the spline.
 		 * For extremely long or curvy splines, this can result in the tokenized grammar being cut short. Alternative subdivision solutions may need to be explored.
@@ -263,7 +264,8 @@ namespace PCGSubdivideSplineHelpers
 			FPCGPoint OutPoint(FTransform(Rotation, Position, Scale), /*InDensity=*/1, PCGHelpers::ComputeSeedFromPosition(Position));
 
 			const double HalfSubmoduleSize = SubdivisionSubmodule.Size * 0.5;
-			OutPoint.SetLocalBounds(FBox(FVector(-HalfSubmoduleSize, 0, 0), FVector(HalfSubmoduleSize, 1, InOutParameters.ModuleHeight)));
+			const double OffsetZ = InOutParameters.Settings->bCenterBounds ? 0.0 : HalfModuleHeight;
+			OutPoint.SetLocalBounds(FBox(FVector(-HalfSubmoduleSize, 0, -HalfModuleHeight + OffsetZ), FVector(HalfSubmoduleSize, 1, HalfModuleHeight + OffsetZ)));
 
 			// Now, handle the metadata attributes
 			if (bHasMetadata)
@@ -307,6 +309,19 @@ FPCGElementPtr UPCGSubdivideSplineSettings::CreateElement() const
 	return MakeShared<FPCGSubdivideSplineElement>();
 }
 
+void UPCGSubdivideSplineSettings::PostLoad()
+{
+	Super::PostLoad();
+	
+#if WITH_EDITOR
+	if (GetLinkerCustomVersion(FFortniteMainBranchObjectVersion::GUID) < FFortniteMainBranchObjectVersion::PCGMeshSamplerDefaultColorPCGSubdivideSplineBoundsFixes)
+	{
+		// For all previous nodes, we'll force this option to false. for retro-compatibility
+		bCenterBounds = false;
+	}
+#endif // WITH_EDITOR
+}
+
 TArray<FPCGPinProperties> UPCGSubdivideSplineSettings::InputPinProperties() const
 {
 	TArray<FPCGPinProperties> PinProperties;
diff --git a/Engine/Plugins/PCG/Source/PCG/Public/Elements/Grammar/PCGSubdivideSpline.h b/Engine/Plugins/PCG/Source/PCG/Public/Elements/Grammar/PCGSubdivideSpline.h
--- a/Engine/Plugins/PCG/Source/PCG/Public/Elements/Grammar/PCGSubdivideSpline.h
+++ b/Engine/Plugins/PCG/Source/PCG/Public/Elements/Grammar/PCGSubdivideSpline.h
@@ -19,6 +19,10 @@ class UPCGSubdivideSplineSettings : public UPCGSubdivisionBaseSettings
 	GENERATED_BODY()
 
 public:
+	//~Begin UObject interface
+	virtual void PostLoad() override;
+	//~End UObject interface
+
 	//~Begin UPCGSettings interface
 #if WITH_EDITOR
 	virtual FName GetDefaultNodeName() const override { return FName(TEXT("SubdivideSpline")); }
@@ -71,6 +75,10 @@ public:
 	UPROPERTY(BlueprintReadWrite, EditAnywhere, Category = Settings, AdvancedDisplay, meta = (PCG_Overridable))
 	bool bSilenceGrammarDoesNotFitErrors = false;
 
+	/** For deprecation purposes only. If this is false, the bounds will be on Z between 0 and Height (previous behavior). If this is true, bounds will be between -Height/2 and Height/2. */
+	UPROPERTY(BlueprintReadWrite, EditAnywhere, Category = Settings, AdvancedDisplay, meta = (EditCondition="!bCenterBounds", EditConditionHides))
+	bool bCenterBounds = true;
+
 	/** Maximum acceptable distance tolerance between placed modules (overlap or gap) - a smaller number is more precise. Affects performance if the number is very small. */
 	UPROPERTY(BlueprintReadWrite, EditAnywhere, Category = Settings, AdvancedDisplay, meta = (PCG_Overridable, ClampMin="0.000001"))
 	double ModulePlacementTolerance = 0.01;
diff --git a/Engine/Plugins/PCGInterops/PCGGeometryScriptInterop/Source/PCGGeometryScriptInterop/Private/Elements/PCGMeshSampler.cpp b/Engine/Plugins/PCGInterops/PCGGeometryScriptInterop/Source/PCGGeometryScriptInterop/Private/Elements/PCGMeshSampler.cpp
--- a/Engine/Plugins/PCGInterops/PCGGeometryScriptInterop/Source/PCGGeometryScriptInterop/Private/Elements/PCGMeshSampler.cpp
+++ b/Engine/Plugins/PCGInterops/PCGGeometryScriptInterop/Source/PCGGeometryScriptInterop/Private/Elements/PCGMeshSampler.cpp
@@ -159,7 +159,7 @@ namespace PCGMeshSampler
 				bool bIsValid = false;
 				UGeometryScriptLibrary_MeshQueryFunctions::ComputeTriangleBarycentricCoords(Context->DynamicMeshes[DataIndex], TriangleId, bIsValid, Position, Dummy1, Dummy2, Dummy3, BarycentricCoord);
 
-				Context->SetPointColorAndDensity(SetPointDensityPtr, TriangleId, BarycentricCoord, OutPoint, DataIndex);
+				Context->SetPointColorAndDensity(SetPointDensityPtr, TriangleId, BarycentricCoord, OutPoint, DataIndex, Settings->bHasDefaultColorOverride ? Settings->DefaultColorOverride : TOptional<FLinearColor>{});
 				Context->SetAttributeValues(Settings->UVChannel, TriangleId, BarycentricCoord, OutPoint.MetadataEntry, DataIndex);
 
 				OutPoint.Seed = PCGHelpers::ComputeSeedFromPosition(OutPoint.Transform.GetLocation());
@@ -255,7 +255,7 @@ namespace PCGMeshSampler
 						OutPoint.Transform = Sample.ToTransform();
 						OutPoint.Steepness = Settings->PointSteepness;
 
-						Context->SetPointColorAndDensity(SetPointDensityPtr, TriangleId, BarycentricCoords, OutPoint, DataIndex);
+						Context->SetPointColorAndDensity(SetPointDensityPtr, TriangleId, BarycentricCoords, OutPoint, DataIndex, Settings->bHasDefaultColorOverride ? Settings->DefaultColorOverride : TOptional<FLinearColor>{});
 						Context->SetAttributeValues(Settings->UVChannel, TriangleId, BarycentricCoords, OutPoint.MetadataEntry, DataIndex);
 
 						OutPoint.Seed = PCGHelpers::ComputeSeedFromPosition(OutPoint.Transform.GetLocation());
@@ -333,20 +333,13 @@ void FPCGMeshSamplerContext::SetAttributeValues(int32 UVChannel, int32 TriangleI
 	}
 }
 
-void FPCGMeshSamplerContext::SetPointColorAndDensity(SetPointDensityFunc SetPointDensityFuncPtr, int32 TriangleId, const FVector& BarycentricCoord, FPCGPoint& OutPoint, int32 DataIndex)
+void FPCGMeshSamplerContext::SetPointColorAndDensity(SetPointDensityFunc SetPointDensityFuncPtr, int32 TriangleId, const FVector& BarycentricCoord, FPCGPoint& OutPoint, int32 DataIndex, TOptional<FLinearColor> DefaultColor)
 {
 	FLinearColor Color;
 	bool bValidVertexColor = false;
-	UGeometryScriptLibrary_MeshQueryFunctions::GetInterpolatedTriangleVertexColor(DynamicMeshes[DataIndex], TriangleId, BarycentricCoord, FLinearColor::White, bValidVertexColor, Color);
-	if (bValidVertexColor)
-	{
-		OutPoint.Color = Color;
-		SetPointDensityFuncPtr(Color, OutPoint);
-	}
-	else
-	{
-		OutPoint.Density = 1.0f;
-	}
+	UGeometryScriptLibrary_MeshQueryFunctions::GetInterpolatedTriangleVertexColor(DynamicMeshes[DataIndex], TriangleId, BarycentricCoord, DefaultColor ? DefaultColor.GetValue() : FLinearColor::White, bValidVertexColor, Color);
+	OutPoint.Color = Color;
+	SetPointDensityFuncPtr(Color, OutPoint);
 }
 
 UPCGMeshSamplerSettings::UPCGMeshSamplerSettings()
@@ -357,6 +350,19 @@ UPCGMeshSamplerSettings::UPCGMeshSamplerSettings()
 	}
 }
 
+void UPCGMeshSamplerSettings::PostLoad()
+{
+	Super::PostLoad();
+	
+#if WITH_EDITOR
+	if (GetLinkerCustomVersion(FFortniteMainBranchObjectVersion::GUID) < FFortniteMainBranchObjectVersion::PCGMeshSamplerDefaultColorPCGSubdivideSplineBoundsFixes)
+	{
+		// For all previous nodes, we'll force this option to false. for retro-compatibility
+		bHasDefaultColorOverride = false;
+	}
+#endif // WITH_EDITOR
+}
+
 TArray<FPCGPinProperties> UPCGMeshSamplerSettings::InputPinProperties() const
 {
 	TArray<FPCGPinProperties> Properties;
@@ -692,10 +698,19 @@ PRAGMA_ENABLE_DEPRECATION_WARNINGS
 		{
 		case EPCGMeshSamplingMethod::OnePointPerVertex:
 		{
-			bool Dummy, Dummy2;
+			bool bHasValidColors, Dummy, Dummy2;
 			UGeometryScriptLibrary_MeshQueryFunctions::GetAllVertexPositions(DynamicMesh, Context->Positions[i], /*bSkipGaps=*/false, /*bHasVertexIDGaps=*/Dummy);
-			UGeometryScriptLibrary_MeshVertexColorFunctions::GetMeshPerVertexColors(DynamicMesh, Context->Colors[i], /*bIsValidColorSet=*/Dummy, /*bHasVertexIDGaps=*/Dummy2);
+			UGeometryScriptLibrary_MeshVertexColorFunctions::GetMeshPerVertexColors(DynamicMesh, Context->Colors[i], /*bIsValidColorSet=*/bHasValidColors, /*bHasVertexIDGaps=*/Dummy2);
 			UGeometryScriptLibrary_MeshNormalsFunctions::GetMeshPerVertexNormals(DynamicMesh, Context->Normals[i], /*bIsValidNormalSet=*/Dummy, /*bHasVertexIDGaps=*/Dummy2);
+
+			if (!bHasValidColors && Settings->bHasDefaultColorOverride)
+			{
+				for (FLinearColor& Color : *(Context->Colors[i].List))
+				{
+					Color = Settings->DefaultColorOverride;
+				}
+			}
+
 			NumIterations = Context->Positions[i].List->Num();
 			break;
 		}
diff --git a/Engine/Plugins/PCGInterops/PCGGeometryScriptInterop/Source/PCGGeometryScriptInterop/Public/Elements/PCGMeshSampler.h b/Engine/Plugins/PCGInterops/PCGGeometryScriptInterop/Source/PCGGeometryScriptInterop/Public/Elements/PCGMeshSampler.h
--- a/Engine/Plugins/PCGInterops/PCGGeometryScriptInterop/Source/PCGGeometryScriptInterop/Public/Elements/PCGMeshSampler.h
+++ b/Engine/Plugins/PCGInterops/PCGGeometryScriptInterop/Source/PCGGeometryScriptInterop/Public/Elements/PCGMeshSampler.h
@@ -61,6 +61,10 @@ class UPCGMeshSamplerSettings : public UPCGSettings
 public:
 	UE_API UPCGMeshSamplerSettings();
 
+	//~Begin UObject interface
+	virtual void PostLoad() override;
+	//~End UObject interface
+
 	//~Begin UPCGSettings interface
 #if WITH_EDITOR
 	virtual FName GetDefaultNodeName() const override { return FName(TEXT("MeshSampler")); }
@@ -104,6 +108,13 @@ public:
 	
 	UPROPERTY(BlueprintReadWrite, EditAnywhere, Category = "Settings|Color & Density", meta = (PCG_Overridable, EditCondition="bUseColorChannelAsDensity"))
 	EPCGColorChannel ColorChannelAsDensity = EPCGColorChannel::Red;
+	
+	UPROPERTY(BlueprintReadWrite, EditAnywhere, Category = "Settings|Color & Density", meta = (InlineEditConditionToggle, PCG_Overridable))
+	bool bHasDefaultColorOverride = true;
+
+	/** When the mesh doesn't have any vertex color, will use this color instead. Otherwise, it will use (0,0,0,0) for "per point" and white for the rest. */
+	UPROPERTY(BlueprintReadWrite, EditAnywhere, Category = "Settings|Color & Density", meta = (PCG_Overridable, EditCondition=bHasDefaultColorOverride))
+	FLinearColor DefaultColorOverride = FLinearColor::White;
 
 	/** Enable voxelisation as a preparation pass. Can be more expensive given the VoxelSize. */
 	UPROPERTY(BlueprintReadWrite, EditAnywhere, Category = "Settings|Voxelize Options", meta = (PCG_Overridable))
@@ -184,7 +195,7 @@ PRAGMA_ENABLE_DEPRECATION_WARNINGS
 	using SetPointDensityFunc = void(*)(const FLinearColor&, FPCGPoint&);
 
 	void SetAttributeValues(int32 UVChannel, int32 TriangleId, const FVector& BarycentricCoord, int64& MetadataEntry, int32 DataIndex);
-	void SetPointColorAndDensity(SetPointDensityFunc SetPointDensityFuncPtr, int32 TriangleId, const FVector& BarycentricCoord, FPCGPoint& OutPoint, int32 DataIndex);
+	void SetPointColorAndDensity(SetPointDensityFunc SetPointDensityFuncPtr, int32 TriangleId, const FVector& BarycentricCoord, FPCGPoint& OutPoint, int32 DataIndex, TOptional<FLinearColor> DefaultColor);
 
 protected:
 	virtual void AddExtraStructReferencedObjects(FReferenceCollector& Collector) override;
diff --git a/Engine/Source/Runtime/Core/Public/UObject/FortniteMainBranchObjectVersions.inl b/Engine/Source/Runtime/Core/Public/UObject/FortniteMainBranchObjectVersions.inl
--- a/Engine/Source/Runtime/Core/Public/UObject/FortniteMainBranchObjectVersions.inl
+++ b/Engine/Source/Runtime/Core/Public/UObject/FortniteMainBranchObjectVersions.inl
@@ -793,5 +793,8 @@ FORTNITE_MAIN_VERSION(AddPhysMatOverrideFlag, 259)
 // Added context DevNotes for translators to FText
 FORTNITE_MAIN_VERSION(AddDevNotesToFText, 260)
 
+// Added default color override for PCG Mesh Sampler and fix bounds in PCG Subdivide Spline
+FORTNITE_MAIN_VERSION(PCGMeshSamplerDefaultColorPCGSubdivideSplineBoundsFixes, 261)
+
 // -----<new versions can be added above this line>-------------------------------------------------
 #undef FORTNITE_MAIN_VERSION
